Fix project categories which have escaped quotes in them. Prevents 4.0 release samples from being grouped incorrectly.

diff --git a/Engine/Source/Runtime/Projects/Private/PluginManagerShared.cpp b/Engine/Source/Runtime/Projects/Private/PluginManagerShared.cpp
--- a/Engine/Source/Runtime/Projects/Private/PluginManagerShared.cpp
+++ b/Engine/Source/Runtime/Projects/Private/PluginManagerShared.cpp
@@ -204,6 +204,13 @@ bool FProjectOrPlugin::DeserializeFromJSON( const FString& JSONInput, FText& Out
 			// Category used to be called CategoryPath in .uplugin files
 			ReadStringFromJSON(FileObject, TEXT("CategoryPath"), ProjectOrPluginInfo.Category);
 		}
+        
+		// Due to a difference in command line parsing between Windows and Mac, we shipped a few Mac samples containing
+		// a category name with escaped quotes. Remove them here to make sure we can list them in the right category.
+		if(ProjectOrPluginInfo.Category.Len() >= 2 && ProjectOrPluginInfo.Category.StartsWith(TEXT("\"")) && ProjectOrPluginInfo.Category.EndsWith(TEXT("\"")))
+		{
+			ProjectOrPluginInfo.Category = ProjectOrPluginInfo.Category.Mid(1, ProjectOrPluginInfo.Category.Len() - 2);
+		}
 
 		ReadStringFromJSON(FileObject, TEXT("CreatedBy"), ProjectOrPluginInfo.CreatedBy);
 		ReadStringFromJSON(FileObject, TEXT("CreatedByURL"), ProjectOrPluginInfo.CreatedByURL);
