Move functions for querying engine installations from Core to DesktopPlatform.

diff --git a/Engine/Source/Developer/DesktopPlatform/Private/DesktopPlatformBase.cpp b/Engine/Source/Developer/DesktopPlatform/Private/DesktopPlatformBase.cpp
new file mode 100644
--- /dev/null
+++ b/Engine/Source/Developer/DesktopPlatform/Private/DesktopPlatformBase.cpp
@@ -0,0 +1,50 @@
+// Copyright 1998-2014 Epic Games, Inc. All Rights Reserved.
+
+#include "DesktopPlatformPrivatePCH.h"
+#include "DesktopPlatformBase.h"
+
+FString FDesktopPlatformBase::GetCurrentEngineIdentifier()
+{
+	FString Identifier;
+	return GetEngineIdentifierFromRootDir(FPlatformMisc::RootDir(), Identifier) ? Identifier : TEXT("");
+}
+
+bool FDesktopPlatformBase::GetEngineRootDirFromIdentifier(const FString &Identifier, FString &OutRootDir)
+{
+	// Get all the installations
+	TMap<FString, FString> Installations;
+	EnumerateEngineInstallations(Installations);
+
+	// Find the one with the right identifier
+	for (TMap<FString, FString>::TConstIterator Iter(Installations); Iter; ++Iter)
+	{
+		if (Iter->Key == Identifier)
+		{
+			OutRootDir = Iter->Value;
+			return true;
+		}
+	}
+	return false;
+}
+
+bool FDesktopPlatformBase::GetEngineIdentifierFromRootDir(const FString &RootDir, FString &OutIdentifier)
+{
+	// Get all the installations
+	TMap<FString, FString> Installations;
+	EnumerateEngineInstallations(Installations);
+
+	// Normalize the root directory
+	FString NormalizedRootDir = RootDir;
+	FPaths::NormalizeDirectoryName(NormalizedRootDir);
+
+	// Find the label for the given directory
+	for (TMap<FString, FString>::TConstIterator Iter(Installations); Iter; ++Iter)
+	{
+		if (Iter->Value == NormalizedRootDir)
+		{
+			OutIdentifier = Iter->Key;
+			return true;
+		}
+	}
+	return false;
+}
diff --git a/Engine/Source/Developer/DesktopPlatform/Private/DesktopPlatformBase.h b/Engine/Source/Developer/DesktopPlatform/Private/DesktopPlatformBase.h
new file mode 100644
--- /dev/null
+++ b/Engine/Source/Developer/DesktopPlatform/Private/DesktopPlatformBase.h
@@ -0,0 +1,14 @@
+// Copyright 1998-2014 Epic Games, Inc. All Rights Reserved.
+
+#pragma once
+
+#include "IDesktopPlatform.h"
+
+class FDesktopPlatformBase : public IDesktopPlatform
+{
+public:
+	// IDesktopPlatform Implementation
+	virtual FString GetCurrentEngineIdentifier() OVERRIDE;
+	virtual bool GetEngineRootDirFromIdentifier(const FString &Identifier, FString &OutRootDir) OVERRIDE;
+	virtual bool GetEngineIdentifierFromRootDir(const FString &RootDir, FString &OutIdentifier) OVERRIDE;
+};
diff --git a/Engine/Source/Developer/DesktopPlatform/Private/Mac/DesktopPlatformMac.cpp b/Engine/Source/Developer/DesktopPlatform/Private/Mac/DesktopPlatformMac.cpp
--- a/Engine/Source/Developer/DesktopPlatform/Private/Mac/DesktopPlatformMac.cpp
+++ b/Engine/Source/Developer/DesktopPlatform/Private/Mac/DesktopPlatformMac.cpp
@@ -512,4 +512,8 @@ bool FDesktopPlatformMac::FileDialogShared(bool bSave, const void* ParentWindowH
 	return bSuccess;
 }
 
+void FDesktopPlatformMac::EnumerateEngineInstallations(TMap<FString, FString> &OutInstallations)
+{
+}
+
 #undef LOCTEXT_NAMESPACE
diff --git a/Engine/Source/Developer/DesktopPlatform/Private/Mac/DesktopPlatformMac.h b/Engine/Source/Developer/DesktopPlatform/Private/Mac/DesktopPlatformMac.h
--- a/Engine/Source/Developer/DesktopPlatform/Private/Mac/DesktopPlatformMac.h
+++ b/Engine/Source/Developer/DesktopPlatform/Private/Mac/DesktopPlatformMac.h
@@ -2,7 +2,9 @@
 
 #pragma once
 
-class FDesktopPlatformMac : public IDesktopPlatform
+#include "../DesktopPlatformBase.h"
+
+class FDesktopPlatformMac : public FDesktopPlatformBase
 {
 public:
 	// IDesktopPlatform Implementation
@@ -12,6 +14,8 @@ public:
 	virtual bool OpenFontDialog(const void* ParentWindowHandle, FString& OutFontName, float& OutHeight, EFontImportFlags::Type& OutFlags) OVERRIDE;
 	virtual bool OpenLauncher(bool Install, const FString& CommandLineParams ) OVERRIDE;
 
+	virtual void EnumerateEngineInstallations(TMap<FString, FString> &OutInstallations) OVERRIDE;
+
 private:
 	bool FileDialogShared(bool bSave, const void* ParentWindowHandle, const FString& DialogTitle, const FString& DefaultPath, const FString& DefaultFile, const FString& FileTypes, uint32 Flags, TArray<FString>& OutFilenames);
 };
diff --git a/Engine/Source/Developer/DesktopPlatform/Private/Windows/DesktopPlatformWindows.cpp b/Engine/Source/Developer/DesktopPlatform/Private/Windows/DesktopPlatformWindows.cpp
--- a/Engine/Source/Developer/DesktopPlatform/Private/Windows/DesktopPlatformWindows.cpp
+++ b/Engine/Source/Developer/DesktopPlatform/Private/Windows/DesktopPlatformWindows.cpp
@@ -345,4 +345,29 @@ bool FDesktopPlatformWindows::FileDialogShared(bool bSave, const void* ParentWin
 	return bSuccess;
 }
 
+void FDesktopPlatformWindows::EnumerateEngineInstallations(TMap<FString, FString> &OutInstallations)
+{
+	HKEY hKey;
+	if (RegOpenKeyEx(HKEY_LOCAL_MACHINE, TEXT("SOFTWARE\\EpicGames\\Unreal Engine"), 0, KEY_READ, &hKey) == ERROR_SUCCESS)
+	{
+		for (::DWORD Index = 0;; Index++)
+		{
+			TCHAR KeyName[256];
+			::DWORD KeyLength = sizeof(KeyName) / sizeof(KeyName[0]);
+			if (RegEnumKeyEx(hKey, Index, KeyName, &KeyLength, NULL, NULL, NULL, NULL) != ERROR_SUCCESS) break;
+
+			TCHAR InstalledDirectory[MAX_PATH];
+			::DWORD InstalledDirectoryLen = sizeof(InstalledDirectory) / sizeof(InstalledDirectory[0]);
+
+			if (RegGetValue(hKey, KeyName, TEXT("InstalledDirectory"), RRF_RT_REG_SZ, NULL, InstalledDirectory, &InstalledDirectoryLen) == ERROR_SUCCESS)
+			{
+				FString NormalizedInstalledDirectory = InstalledDirectory;
+				FPaths::NormalizeDirectoryName(NormalizedInstalledDirectory);
+				OutInstallations.Add(KeyName, NormalizedInstalledDirectory);
+			}
+		}
+		RegCloseKey(hKey);
+	}
+}
+
 #undef LOCTEXT_NAMESPACE
diff --git a/Engine/Source/Developer/DesktopPlatform/Private/Windows/DesktopPlatformWindows.h b/Engine/Source/Developer/DesktopPlatform/Private/Windows/DesktopPlatformWindows.h
--- a/Engine/Source/Developer/DesktopPlatform/Private/Windows/DesktopPlatformWindows.h
+++ b/Engine/Source/Developer/DesktopPlatform/Private/Windows/DesktopPlatformWindows.h
@@ -2,7 +2,9 @@
 
 #pragma once
 
-class FDesktopPlatformWindows : public IDesktopPlatform
+#include "../DesktopPlatformBase.h"
+
+class FDesktopPlatformWindows : public FDesktopPlatformBase
 {
 public:
 	// IDesktopPlatform Implementation
@@ -12,6 +14,8 @@ public:
 	virtual bool OpenFontDialog(const void* ParentWindowHandle, FString& OutFontName, float& OutHeight, EFontImportFlags::Type& OutFlags) OVERRIDE;
 	virtual bool OpenLauncher(bool Install, const FString& CommandLineParams ) OVERRIDE;
 
+	virtual void EnumerateEngineInstallations(TMap<FString, FString> &OutInstallations) OVERRIDE;
+
 private:
 	bool FileDialogShared(bool bSave, const void* ParentWindowHandle, const FString& DialogTitle, const FString& DefaultPath, const FString& DefaultFile, const FString& FileTypes, uint32 Flags, TArray<FString>& OutFilenames);
 };
diff --git a/Engine/Source/Developer/DesktopPlatform/Public/IDesktopPlatform.h b/Engine/Source/Developer/DesktopPlatform/Public/IDesktopPlatform.h
--- a/Engine/Source/Developer/DesktopPlatform/Public/IDesktopPlatform.h
+++ b/Engine/Source/Developer/DesktopPlatform/Public/IDesktopPlatform.h
@@ -116,4 +116,35 @@ public:
 	 * @return true if the marketplace was opened, false if it is not installed or could not be installed/opened.
 	 */
 	virtual bool OpenLauncher(bool Install, const FString& CommandLineParams ) = 0;
-};
\ No newline at end of file
+
+	/**
+	* Gets the identifier for the currently executing engine installation
+	*
+	* @return	Identifier for the current engine installation. Empty string if it isn't registered.
+	*/
+	virtual FString GetCurrentEngineIdentifier() = 0;
+
+	/**
+	* Enumerates all the registered engine installations.
+	*
+	* @param	OutInstallations	Array which is filled in with identifier/root-directory pairs for all known installations. Identifiers are typically
+	*								version strings for canonical UE4 releases or GUID strings for GitHub releases.
+	*/
+	virtual void EnumerateEngineInstallations(TMap<FString, FString> &OutInstallations) = 0;
+
+	/**
+	* Returns the identifier for the engine with the given root directory.
+	*
+	* @param	RootDirName			Root directory for the engine.
+	* @param	OutIdentifier		Identifier used to refer to this installation.
+	*/
+	virtual bool GetEngineRootDirFromIdentifier(const FString &Identifier, FString &OutRootDir) = 0;
+
+	/**
+	* Returns the identifier for the engine with the given root directory.
+	*
+	* @param	RootDirName			Root directory for the engine.
+	* @param	OutIdentifier		Identifier used to refer to this installation.
+	*/
+	virtual bool GetEngineIdentifierFromRootDir(const FString &RootDir, FString &OutIdentifier) = 0;
+};
diff --git a/Engine/Source/Editor/GameProjectGeneration/GameProjectGeneration.Build.cs b/Engine/Source/Editor/GameProjectGeneration/GameProjectGeneration.Build.cs
--- a/Engine/Source/Editor/GameProjectGeneration/GameProjectGeneration.Build.cs
+++ b/Engine/Source/Editor/GameProjectGeneration/GameProjectGeneration.Build.cs
@@ -28,6 +28,7 @@ public class GameProjectGeneration : ModuleRules
                 "EditorStyle",
                 "SourceControl",
                 "UnrealEd",
+				"DesktopPlatform"
 			}
 		);
 
diff --git a/Engine/Source/Editor/GameProjectGeneration/Private/GameProjectUtils.cpp b/Engine/Source/Editor/GameProjectGeneration/Private/GameProjectUtils.cpp
--- a/Engine/Source/Editor/GameProjectGeneration/Private/GameProjectUtils.cpp
+++ b/Engine/Source/Editor/GameProjectGeneration/Private/GameProjectUtils.cpp
@@ -10,6 +10,8 @@
 #include "EngineAnalytics.h"
 #include "EngineBuildSettings.h"
 
+#include "DesktopPlatformModule.h"
+
 #define LOCTEXT_NAMESPACE "GameProjectUtils"
 
 #define MAX_PROJECT_PATH_BUFFER_SPACE 130 // Leave a reasonable buffer of additional characters to account for files created in the content directory during or after project generation
@@ -371,7 +373,7 @@ void GameProjectUtils::CheckForOutOfDateGameProjectFile()
 	if ( !LoadedProjectFilePath.IsEmpty() )
 	{
 		FProjectStatus ProjectStatus;
-		if ( IProjectManager::Get().QueryStatusForProject(LoadedProjectFilePath, ProjectStatus) )
+		if (IProjectManager::Get().QueryStatusForProject(LoadedProjectFilePath, FDesktopPlatformModule::Get()->GetCurrentEngineIdentifier(), ProjectStatus))
 		{
 			if ( !ProjectStatus.bUpToDate )
 			{
@@ -411,7 +413,7 @@ bool GameProjectUtils::UpdateGameProject()
 	if ( !ProjectFilename.IsEmpty() )
 	{
 		FProjectStatus ProjectStatus;
-		if ( IProjectManager::Get().QueryStatusForProject(ProjectFilename, ProjectStatus) )
+		if ( IProjectManager::Get().QueryStatusForProject(ProjectFilename, FDesktopPlatformModule::Get()->GetCurrentEngineIdentifier(), ProjectStatus) )
 		{
 			if ( ProjectStatus.bUpToDate )
 			{
@@ -607,7 +609,7 @@ bool GameProjectUtils::GenerateProjectFromScratch(const FString& NewProjectFile,
 	// Generate the project file
 	{
 		FText LocalFailReason;
-		if ( IProjectManager::Get().GenerateNewProjectFile(NewProjectFile, StartupModuleNames, LocalFailReason) )
+		if (IProjectManager::Get().GenerateNewProjectFile(NewProjectFile, StartupModuleNames, FDesktopPlatformModule::Get()->GetCurrentEngineIdentifier(), LocalFailReason))
 		{
 			CreatedFiles.Add(NewProjectFile);
 		}
@@ -976,7 +978,7 @@ bool GameProjectUtils::CreateProjectFromTemplate(const FString& NewProjectFile,
 	// Generate the project file
 	{
 		FText LocalFailReason;
-		if ( IProjectManager::Get().DuplicateProjectFile(TemplateFile, NewProjectFile, LocalFailReason) )
+		if (IProjectManager::Get().DuplicateProjectFile(TemplateFile, NewProjectFile, FDesktopPlatformModule::Get()->GetCurrentEngineIdentifier(), LocalFailReason))
 		{
 			CreatedFiles.Add(NewProjectFile);
 		}
@@ -1796,7 +1798,7 @@ bool GameProjectUtils::UpdateGameProjectFile(const FString& ProjectFilename, con
 	}
 
 	// Now tell the project manager to update the file
-	if ( !IProjectManager::Get().UpdateLoadedProjectFileToCurrent(StartupModuleNames, OutFailReason) )
+	if (!IProjectManager::Get().UpdateLoadedProjectFileToCurrent(StartupModuleNames, FDesktopPlatformModule::Get()->GetCurrentEngineIdentifier(), OutFailReason))
 	{
 		return false;
 	}
diff --git a/Engine/Source/Editor/GameProjectGeneration/Private/SProjectBrowser.cpp b/Engine/Source/Editor/GameProjectGeneration/Private/SProjectBrowser.cpp
--- a/Engine/Source/Editor/GameProjectGeneration/Private/SProjectBrowser.cpp
+++ b/Engine/Source/Editor/GameProjectGeneration/Private/SProjectBrowser.cpp
@@ -611,7 +611,7 @@ void SProjectBrowser::FindProjects(bool bAllowProjectCreate)
 		{
 			const bool bPromptIfSavedWithNewerVersionOfEngine = false;
 			FProjectStatus ProjectStatus;
-			if ( IProjectManager::Get().QueryStatusForProject(ProjectFilename, ProjectStatus) )
+			if (IProjectManager::Get().QueryStatusForProject(ProjectFilename, FDesktopPlatformModule::Get()->GetCurrentEngineIdentifier(), ProjectStatus))
 			{
 				// @todo localized project name
 				const FText ProjectName = FText::FromString(ProjectStatus.Name);
diff --git a/Engine/Source/Editor/UnrealEd/Private/UnrealEdMisc.cpp b/Engine/Source/Editor/UnrealEd/Private/UnrealEdMisc.cpp
--- a/Engine/Source/Editor/UnrealEd/Private/UnrealEdMisc.cpp
+++ b/Engine/Source/Editor/UnrealEd/Private/UnrealEdMisc.cpp
@@ -30,6 +30,7 @@
 #include "NormalMapIdentification.h"
 #include "EngineBuildSettings.h"
 #include "Slate.h"
+#include "DesktopPlatformModule.h"
 
 #define LOCTEXT_NAMESPACE "UnrealEd"
 
@@ -335,7 +336,7 @@ void FUnrealEdMisc::InitEngineAnalytics()
 			const FString& LoadedProjectFilePath = FPaths::GetProjectFilePath();
 			FProjectStatus ProjectStatus;
 
-			if ( IProjectManager::Get().QueryStatusForProject(LoadedProjectFilePath, ProjectStatus) )
+			if (IProjectManager::Get().QueryStatusForProject(LoadedProjectFilePath, FDesktopPlatformModule::Get()->GetCurrentEngineIdentifier(), ProjectStatus))
 			{
 				if ( ProjectStatus.bSignedSampleProject )
 				{
diff --git a/Engine/Source/Programs/UnrealVersionSelector/Private/FileAssociation.cpp b/Engine/Source/Programs/UnrealVersionSelector/Private/FileAssociation.cpp
--- a/Engine/Source/Programs/UnrealVersionSelector/Private/FileAssociation.cpp
+++ b/Engine/Source/Programs/UnrealVersionSelector/Private/FileAssociation.cpp
@@ -2,6 +2,7 @@
 #include "FileAssociation.h"
 #include "PlatformInstallation.h"
 #include "Runtime/Core/Public/Serialization/Json/Json.h"
+#include "DesktopPlatformModule.h"
 
 TSharedPtr<FJsonObject> LoadJson(const FString &FileName)
 {
@@ -71,7 +72,7 @@ bool GetEngineIdForProject(const FString &ProjectFileName, FString &OutId)
 bool GetEngineRootDirForProject(const FString &ProjectFileName, FString &OutRootDir)
 {
 	FString Id;
-	if (GetEngineIdForProject(ProjectFileName, Id) && FPlatformMisc::GetEngineRootDirFromIdentifier(Id, OutRootDir))
+	if (GetEngineIdForProject(ProjectFileName, Id) && FDesktopPlatformModule::Get()->GetEngineRootDirFromIdentifier(Id, OutRootDir))
 	{
 		return true;
 	}
diff --git a/Engine/Source/Programs/UnrealVersionSelector/Private/PlatformInstallation.cpp b/Engine/Source/Programs/UnrealVersionSelector/Private/PlatformInstallation.cpp
--- a/Engine/Source/Programs/UnrealVersionSelector/Private/PlatformInstallation.cpp
+++ b/Engine/Source/Programs/UnrealVersionSelector/Private/PlatformInstallation.cpp
@@ -1,5 +1,6 @@
 #include "UnrealVersionSelector.h"
 #include "PlatformInstallation.h"
+#include "DesktopPlatformModule.h"
 
 int32 ParseReleaseVersion(const FString &Version)
 {
@@ -22,23 +23,23 @@ int32 ParseReleaseVersion(const FString &Version)
 
 bool GetDefaultEngineId(FString &OutId)
 {
-	TArray<TPair<FString, FString>> Installations;
-	FPlatformMisc::EnumerateEngineInstallations(Installations);
+	TMap<FString, FString> Installations;
+	FDesktopPlatformModule::Get()->EnumerateEngineInstallations(Installations);
 
 	bool bRes = false;
 	if (Installations.Num() > 0)
 	{
 		// Default to the first install
-		OutId = Installations[0].Key;
+		OutId = TMap<FString, FString>::TConstIterator(Installations).Key();
 
 		// Try to find the highest versioned install
 		int32 BestVersion = INDEX_NONE;
-		for (TArray<TPair<FString, FString>>::TConstIterator Iter(Installations); Iter; ++Iter)
+		for (TMap<FString, FString>::TConstIterator Iter(Installations); Iter; ++Iter)
 		{
-			int32 Version = ParseReleaseVersion(Iter->Key);
+			int32 Version = ParseReleaseVersion(Iter.Key());
 			if (Version > BestVersion)
 			{
-				OutId = Iter->Key;
+				OutId = Iter.Key();
 				BestVersion = Version;
 			}
 		}
@@ -49,5 +50,5 @@ bool GetDefaultEngineId(FString &OutId)
 bool GetDefaultEngineRootDir(FString &OutDirName)
 {
 	FString Id;
-	return GetDefaultEngineId(Id) && FPlatformMisc::GetEngineRootDirFromIdentifier(Id, OutDirName);
+	return GetDefaultEngineId(Id) && FDesktopPlatformModule::Get()->GetEngineRootDirFromIdentifier(Id, OutDirName);
 }
diff --git a/Engine/Source/Programs/UnrealVersionSelector/Private/UnrealVersionSelector.cpp b/Engine/Source/Programs/UnrealVersionSelector/Private/UnrealVersionSelector.cpp
--- a/Engine/Source/Programs/UnrealVersionSelector/Private/UnrealVersionSelector.cpp
+++ b/Engine/Source/Programs/UnrealVersionSelector/Private/UnrealVersionSelector.cpp
@@ -39,16 +39,16 @@ bool RegisterCurrentEngineDirectory()
 void EnumerateInvalidInstallations(TArray<FString> &OutIds)
 {
 	// Enumerate all the engine installations
-	TArray< TPair<FString, FString> > Installations;
-	FPlatformMisc::EnumerateEngineInstallations(Installations);
+	TMap<FString, FString> Installations;
+	FDesktopPlatformModule::Get()->EnumerateEngineInstallations(Installations);
 
 	// Check each one of them has a valid engine path
-	for (TArray<TPair<FString, FString>>::TConstIterator Iter(Installations); Iter; ++Iter)
+	for (TMap<FString, FString>::TConstIterator Iter(Installations); Iter; ++Iter)
 	{
-		const FString &RootDir = Iter->Value;
+		const FString &RootDir = Iter.Value();
 		if (!FPlatformInstallation::IsValidEngineRootDir(RootDir))
 		{
-			OutIds.Add(Iter->Key);
+			OutIds.Add(Iter.Key());
 		}
 	}
 }
@@ -111,7 +111,7 @@ bool RegisterEngineDirectory(const FString &EngineDir)
 {
 	// Get any existing tag name or register a new one
 	FString EngineId;
-	if (!FPlatformMisc::GetEngineIdentifierFromRootDir(EngineDir, EngineId))
+	if (!FDesktopPlatformModule::Get()->GetEngineIdentifierFromRootDir(EngineDir, EngineId))
 	{
 		EngineId = FGuid::NewGuid().ToString(EGuidFormats::DigitsWithHyphensInBraces);
 	}
@@ -128,6 +128,9 @@ bool SelectEngineAssociation(const FString &ProjectFileName, const FString &Engi
 		FPlatformMisc::MessageBoxExt(EAppMsgType::Ok, TEXT("Couldn't set association for project. Check the file is writeable."), TEXT("Error"));
 		return false;
 	}
+
+	// Notify the user that it's been changed
+	FPlatformMisc::MessageBoxExt(EAppMsgType::Ok, TEXT("Updated project file."), TEXT("Success"));
 	return true;
 }
 
@@ -153,14 +156,14 @@ bool BrowseForEngineAssociation(const FString &ProjectFileName)
 
 	// Check that it's a registered engine directory
 	FString EngineId;
-	if (!FPlatformMisc::GetEngineIdentifierFromRootDir(NewEngineRootDir, EngineId))
+	if (!FDesktopPlatformModule::Get()->GetEngineIdentifierFromRootDir(NewEngineRootDir, EngineId))
 	{
 		int32 ExitCode;
 		if (!FPlatformProcess::ExecElevatedProcess(*FullExecutableName(), *FString::Printf(TEXT("/register \"%s\""), *NewEngineRootDir), &ExitCode) || ExitCode != 0)
 		{
 			return false;
 		}
-		if (!FPlatformMisc::GetEngineIdentifierFromRootDir(NewEngineRootDir, EngineId))
+		if (!FDesktopPlatformModule::Get()->GetEngineIdentifierFromRootDir(NewEngineRootDir, EngineId))
 		{
 			return false;
 		}
diff --git a/Engine/Source/Programs/UnrealVersionSelector/Private/Windows/WindowsPlatformInstallation.cpp b/Engine/Source/Programs/UnrealVersionSelector/Private/Windows/WindowsPlatformInstallation.cpp
--- a/Engine/Source/Programs/UnrealVersionSelector/Private/Windows/WindowsPlatformInstallation.cpp
+++ b/Engine/Source/Programs/UnrealVersionSelector/Private/Windows/WindowsPlatformInstallation.cpp
@@ -110,8 +110,8 @@ void GetRequiredRegistrySettings(TIndirectArray<FRegistryRootedKey> &RootedKeys)
 	FRegistryKey *SwitchShellKey = SwitchKey->FindOrAddKey(TEXT("shell"));
 
 	// HKLM\SOFTWARE\Classes\Unreal.ProjectFile\switchversion\shell\<Label>
-	TArray<TPair<FString, FString>> Installations;
-	FPlatformMisc::EnumerateEngineInstallations(Installations);
+	TMap<FString, FString> Installations;
+	FDesktopPlatformModule::Get()->EnumerateEngineInstallations(Installations);
 
 	if (Installations.Num() > 0)
 	{
@@ -122,10 +122,10 @@ void GetRequiredRegistrySettings(TIndirectArray<FRegistryRootedKey> &RootedKeys)
 
 		// Build a list of installation names
 		TMap<FString, FString> InstallationNames;
-		for (TArray<TPair<FString, FString>>::TConstIterator Iter(Installations); Iter; ++Iter)
+		for (TMap<FString, FString>::TConstIterator Iter(Installations); Iter; ++Iter)
 		{
 			const FString &Id = Iter->Key;
-			FString Description = GetInstallationDescription(Id, Iter->Value);
+			FString Description = GetInstallationDescription(Id, Iter.Value());
 			InstallationNames.Add(Description, Id);
 		}
 		InstallationNames.KeySort(FEngineLabelSortPredicate());
diff --git a/Engine/Source/Runtime/Core/Private/GenericPlatform/GenericPlatformMisc.cpp b/Engine/Source/Runtime/Core/Private/GenericPlatform/GenericPlatformMisc.cpp
--- a/Engine/Source/Runtime/Core/Private/GenericPlatform/GenericPlatformMisc.cpp
+++ b/Engine/Source/Runtime/Core/Private/GenericPlatform/GenericPlatformMisc.cpp
@@ -694,58 +694,3 @@ FText FGenericPlatformMisc::GetFileManagerName()
 	return NSLOCTEXT("GenericPlatform", "FileManagerName", "File Manager");
 }
 
-void FGenericPlatformMisc::EnumerateEngineInstallations(TArray< TPair<FString, FString> > &OutInstallations)
-{
-#if PLATFORM_DESKTOP
-	// Read the config file.
-	FConfigFile ConfigFile;
-	ConfigFile.Read(FString(FPlatformProcess::ApplicationSettingsDir()) / FString(EPIC_PRODUCT_IDENTIFIER) / FString(TEXT("Install.ini")));
-
-	// Enumerate the installation list.
-	FConfigSection &Section = ConfigFile.FindOrAdd(TEXT("Installations"));
-	for (TMultiMap<FName, FString>::TConstIterator Iter(Section); Iter; ++Iter)
-	{
-		OutInstallations.Add(TPairInitializer<FString, FString>(Iter.Key().ToString(), Iter.Value()));
-	}
-#endif
-}
-
-bool FGenericPlatformMisc::GetEngineRootDirFromIdentifier(const FString &Identifier, FString &OutRootDir)
-{
-	// Get all the installations
-	TArray< TPair<FString, FString> > Installations;
-	FPlatformMisc::EnumerateEngineInstallations(Installations);
-
-	// Find the one with the right identifier
-	for (TArray< TPair<FString, FString> >::TConstIterator Iter(Installations); Iter; ++Iter)
-	{
-		if (Iter->Key == Identifier)
-		{
-			OutRootDir = Iter->Value;
-			return true;
-		}
-	}
-	return false;
-}
-
-bool FGenericPlatformMisc::GetEngineIdentifierFromRootDir(const FString &RootDir, FString &OutIdentifier)
-{
-	// Get all the installations
-	TArray< TPair<FString, FString> > Installations;
-	FPlatformMisc::EnumerateEngineInstallations(Installations);
-
-	// Normalize the root directory
-	FString NormalizedRootDir = RootDir;
-	FPaths::NormalizeDirectoryName(NormalizedRootDir);
-
-	// Find the label for the given directory
-	for (TArray< TPair<FString, FString> >::TConstIterator Iter(Installations); Iter; ++Iter)
-	{
-		if (Iter->Value == NormalizedRootDir)
-		{
-			OutIdentifier = Iter->Key;
-			return true;
-		}
-	}
-	return false;
-}
diff --git a/Engine/Source/Runtime/Core/Private/Windows/WindowsPlatformMisc.cpp b/Engine/Source/Runtime/Core/Private/Windows/WindowsPlatformMisc.cpp
--- a/Engine/Source/Runtime/Core/Private/Windows/WindowsPlatformMisc.cpp
+++ b/Engine/Source/Runtime/Core/Private/Windows/WindowsPlatformMisc.cpp
@@ -1852,28 +1852,3 @@ FText FWindowsPlatformMisc::GetFileManagerName()
 {
 	return NSLOCTEXT("WindowsPlatform", "FileManagerName", "Explorer");
 }
-
-void FWindowsPlatformMisc::EnumerateEngineInstallations(TArray< TPair<FString, FString> > &OutInstallations)
-{
-	HKEY hKey;
-	if (RegOpenKeyEx(HKEY_LOCAL_MACHINE, TEXT("SOFTWARE\\EpicGames\\Unreal Engine"), 0, KEY_READ, &hKey) == ERROR_SUCCESS)
-	{
-		for (::DWORD Index = 0;; Index++)
-		{
-			TCHAR KeyName[256];
-			::DWORD KeyLength = sizeof(KeyName) / sizeof(KeyName[0]);
-			if (RegEnumKeyEx(hKey, Index, KeyName, &KeyLength, NULL, NULL, NULL, NULL) != ERROR_SUCCESS) break;
-
-			TCHAR InstalledDirectory[MAX_PATH];
-			::DWORD InstalledDirectoryLen = sizeof(InstalledDirectory) / sizeof(InstalledDirectory[0]);
-
-			if (RegGetValue(hKey, KeyName, TEXT("InstalledDirectory"), RRF_RT_REG_SZ, NULL, InstalledDirectory, &InstalledDirectoryLen) == ERROR_SUCCESS)
-			{
-				FString NormalizedInstalledDirectory = InstalledDirectory;
-				FPaths::NormalizeDirectoryName(NormalizedInstalledDirectory);
-				OutInstallations.Add(TPairInitializer<FString, FString>(KeyName, NormalizedInstalledDirectory));
-			}
-		}
-		RegCloseKey(hKey);
-	}
-}
diff --git a/Engine/Source/Runtime/Core/Public/GenericPlatform/GenericPlatformMisc.h b/Engine/Source/Runtime/Core/Public/GenericPlatform/GenericPlatformMisc.h
--- a/Engine/Source/Runtime/Core/Public/GenericPlatform/GenericPlatformMisc.h
+++ b/Engine/Source/Runtime/Core/Public/GenericPlatform/GenericPlatformMisc.h
@@ -726,30 +726,6 @@ public:
 		return PLATFORM_HAS_TOUCH_MAIN_SCREEN;
 	}
 
-	/**
-	* Enumerates all the registered engine installations.
-	*
-	* @param	OutInstallations	Array which is filled in with identifier/root-directory pairs for all known installations. Identifiers are typically
-	*								version strings for canonical UE4 releases or GUID strings for GitHub releases.
-	*/
-	static void EnumerateEngineInstallations(TArray< TPair<FString, FString> > &OutInstallations);
-
-	/**
-	* Returns the identifier for the engine with the given root directory.
-	*
-	* @param	RootDirName			Root directory for the engine.
-	* @param	OutIdentifier		Identifier used to refer to this installation.
-	*/
-	static bool GetEngineRootDirFromIdentifier(const FString &Identifier, FString &OutRootDir);
-
-	/**
-	* Returns the identifier for the engine with the given root directory.
-	*
-	* @param	RootDirName			Root directory for the engine.
-	* @param	OutIdentifier		Identifier used to refer to this installation.
-	*/
-	static bool GetEngineIdentifierFromRootDir(const FString &RootDir, FString &OutIdentifier);
-
 #if !UE_BUILD_SHIPPING
 protected:
 	/** Whether the user should be prompted to allow for a remote debugger to be attached */
diff --git a/Engine/Source/Runtime/Core/Public/Windows/WindowsPlatformMisc.h b/Engine/Source/Runtime/Core/Public/Windows/WindowsPlatformMisc.h
--- a/Engine/Source/Runtime/Core/Public/Windows/WindowsPlatformMisc.h
+++ b/Engine/Source/Runtime/Core/Public/Windows/WindowsPlatformMisc.h
@@ -189,9 +189,6 @@ struct CORE_API FWindowsPlatformMisc : public FGenericPlatformMisc
 
 	/** @return Get the name of the platform specific file manager (Explorer) */
 	static FText GetFileManagerName();
-
-	/** Enumerate all the engine installations, using the registry */
-	static void EnumerateEngineInstallations(TArray< TPair<FString, FString> > &OutInstallations);
 };
 
 typedef FWindowsPlatformMisc FPlatformMisc;
diff --git a/Engine/Source/Runtime/Projects/Private/PluginManagerShared.cpp b/Engine/Source/Runtime/Projects/Private/PluginManagerShared.cpp
--- a/Engine/Source/Runtime/Projects/Private/PluginManagerShared.cpp
+++ b/Engine/Source/Runtime/Projects/Private/PluginManagerShared.cpp
@@ -296,7 +296,7 @@ FString FProjectOrPlugin::SerializeToJSON( ) const
 	return JSONOutput;
 }
 
-bool FProjectOrPlugin::IsUpToDate( ) const
+bool FProjectOrPlugin::IsUpToDate( const FString &EngineIdentifier ) const
 {
 	const FProjectOrPluginInfo& ProjectOrPluginInfo = GetProjectOrPluginInfo();
 
@@ -305,10 +305,7 @@ bool FProjectOrPlugin::IsUpToDate( ) const
 		return false;
 	}
 
-	FString CurrentEngineIdentifier;
-	FPlatformMisc::GetEngineIdentifierFromRootDir(FPlatformMisc::RootDir(), CurrentEngineIdentifier);
-
-	if (ProjectOrPluginInfo.EngineAssociation != CurrentEngineIdentifier)
+	if (ProjectOrPluginInfo.EngineAssociation != EngineIdentifier)
 	{
 		return false;
 	}
@@ -316,7 +313,7 @@ bool FProjectOrPlugin::IsUpToDate( ) const
 	return true;
 }
 
-void FProjectOrPlugin::UpdateVersionToCurrent( )
+void FProjectOrPlugin::UpdateVersionToCurrent( const FString &EngineIdentifier )
 {
 	FProjectOrPluginInfo& ProjectOrPluginInfo = GetProjectOrPluginInfo();
 
@@ -324,8 +321,7 @@ void FProjectOrPlugin::UpdateVersionToCurrent( )
 	ProjectOrPluginInfo.EngineVersion = GEngineVersion;
 	ProjectOrPluginInfo.PackageFileUE4Version = GPackageFileUE4Version;
 	ProjectOrPluginInfo.PackageFileLicenseeUE4Version = GPackageFileLicenseeUE4Version;
-
-	FPlatformMisc::GetEngineIdentifierFromRootDir(FPlatformMisc::RootDir(), ProjectOrPluginInfo.EngineAssociation);
+	ProjectOrPluginInfo.EngineAssociation = EngineIdentifier;
 }
 
 void FProjectOrPlugin::ReplaceModulesInProject(const TArray<FString>* StartupModuleNames)
diff --git a/Engine/Source/Runtime/Projects/Private/PluginManagerShared.h b/Engine/Source/Runtime/Projects/Private/PluginManagerShared.h
--- a/Engine/Source/Runtime/Projects/Private/PluginManagerShared.h
+++ b/Engine/Source/Runtime/Projects/Private/PluginManagerShared.h
@@ -221,12 +221,12 @@ public:
 	 *
 	 * @return true if the project file is up to date, false otherwise.
 	 */
-	bool IsUpToDate( ) const;
+	bool IsUpToDate( const FString &EngineIdentifier ) const;
 
 	/**
 	 * Updates all version info to match the currently running executable.
 	 */
-	void UpdateVersionToCurrent( );
+	void UpdateVersionToCurrent( const FString &EngineIdentifier );
 
 	/**
 	 * Updates the modules in the project file
diff --git a/Engine/Source/Runtime/Projects/Private/ProjectManager.cpp b/Engine/Source/Runtime/Projects/Private/ProjectManager.cpp
--- a/Engine/Source/Runtime/Projects/Private/ProjectManager.cpp
+++ b/Engine/Source/Runtime/Projects/Private/ProjectManager.cpp
@@ -179,10 +179,10 @@ const FString& FProjectManager::NonStaticGetProjectFileExtension()
 	return GameProjectFileExtension;
 }
 
-bool FProjectManager::GenerateNewProjectFile(const FString& NewProjectFilename, const TArray<FString>& StartupModuleNames, FText& OutFailReason)
+bool FProjectManager::GenerateNewProjectFile(const FString& NewProjectFilename, const TArray<FString>& StartupModuleNames, const FString& EngineIdentifier, FText& OutFailReason)
 {
 	TSharedRef<FProject> NewProject = MakeShareable( new FProject() );
-	NewProject->UpdateVersionToCurrent();
+	NewProject->UpdateVersionToCurrent(EngineIdentifier);
 	NewProject->ReplaceModulesInProject(&StartupModuleNames);
 
 	const FString& FileContents = NewProject->SerializeToJSON();
@@ -197,7 +197,7 @@ bool FProjectManager::GenerateNewProjectFile(const FString& NewProjectFilename,
 	}
 }
 
-bool FProjectManager::DuplicateProjectFile(const FString& SourceProjectFilename, const FString& NewProjectFilename, FText& OutFailReason)
+bool FProjectManager::DuplicateProjectFile(const FString& SourceProjectFilename, const FString& NewProjectFilename, const FString& EngineIdentifier, FText& OutFailReason)
 {
 	// Load the source project
 	TSharedRef<FProject> SourceProject = MakeShareable( new FProject() );
@@ -223,7 +223,7 @@ bool FProjectManager::DuplicateProjectFile(const FString& SourceProjectFilename,
 
 	// Create new project, update version numbers (no need to replace modules here)
 	TSharedRef<FProject> NewProject = MakeShareable( new FProject(ProjectInfo) );
-	NewProject->UpdateVersionToCurrent();
+	NewProject->UpdateVersionToCurrent(EngineIdentifier);
 
 	// Serialize and write to disk
 	const FString& FileContents = NewProject->SerializeToJSON();
@@ -238,7 +238,7 @@ bool FProjectManager::DuplicateProjectFile(const FString& SourceProjectFilename,
 	}
 }
 
-bool FProjectManager::UpdateLoadedProjectFileToCurrent(const TArray<FString>* StartupModuleNames, FText& OutFailReason)
+bool FProjectManager::UpdateLoadedProjectFileToCurrent(const TArray<FString>* StartupModuleNames, const FString& EngineIdentifier, FText& OutFailReason)
 {
 	if ( !CurrentlyLoadedProject.IsValid() )
 	{
@@ -246,7 +246,7 @@ bool FProjectManager::UpdateLoadedProjectFileToCurrent(const TArray<FString>* St
 	}
 
 	// Freshen version information
-	CurrentlyLoadedProject->UpdateVersionToCurrent();
+	CurrentlyLoadedProject->UpdateVersionToCurrent(EngineIdentifier);
 
 	// Replace the modules names, if specified
 	CurrentlyLoadedProject->ReplaceModulesInProject(StartupModuleNames);
@@ -293,7 +293,7 @@ bool FProjectManager::SignSampleProject(const FString& FilePath, const FString&
 	}
 }
 
-bool FProjectManager::QueryStatusForProject(const FString& FilePath, FProjectStatus& OutProjectStatus) const
+bool FProjectManager::QueryStatusForProject(const FString& FilePath, const FString& EngineIdentifier, FProjectStatus& OutProjectStatus) const
 {
 	TSharedRef<FProject> NewProject = MakeShareable( new FProject() );
 	FText FailReason;
@@ -313,7 +313,7 @@ bool FProjectManager::QueryStatusForProject(const FString& FilePath, FProjectSta
 	OutProjectStatus.Description = ProjectInfo.Description;
 	OutProjectStatus.Category = ProjectInfo.Category;
 	OutProjectStatus.bSignedSampleProject = NewProject->IsSignedSampleProject(FilePath);
-	OutProjectStatus.bUpToDate = NewProject->IsUpToDate();
+	OutProjectStatus.bUpToDate = NewProject->IsUpToDate(EngineIdentifier);
 
 	return true;
 }
diff --git a/Engine/Source/Runtime/Projects/Private/ProjectManager.h b/Engine/Source/Runtime/Projects/Private/ProjectManager.h
--- a/Engine/Source/Runtime/Projects/Private/ProjectManager.h
+++ b/Engine/Source/Runtime/Projects/Private/ProjectManager.h
@@ -88,11 +88,11 @@ public:
 	virtual bool LoadModulesForProject( const ELoadingPhase::Type LoadingPhase ) OVERRIDE;
 	virtual const FString& GetAutoLoadProjectFileName() OVERRIDE;
 	virtual const FString& NonStaticGetProjectFileExtension() OVERRIDE;
-	virtual bool GenerateNewProjectFile(const FString& NewProjectFilename, const TArray<FString>& ModuleNames, FText& OutFailReason) OVERRIDE;
-	virtual bool DuplicateProjectFile(const FString& SourceProjectFilename, const FString& NewProjectFilename, FText& OutFailReason) OVERRIDE;
-	virtual bool UpdateLoadedProjectFileToCurrent(const TArray<FString>* StartupModuleNames, FText& OutFailReason) OVERRIDE;
+	virtual bool GenerateNewProjectFile(const FString& NewProjectFilename, const TArray<FString>& ModuleNames, const FString& EngineIdentifier, FText& OutFailReason) OVERRIDE;
+	virtual bool DuplicateProjectFile(const FString& SourceProjectFilename, const FString& NewProjectFilename, const FString& EngineIdentifier, FText& OutFailReason) OVERRIDE;
+	virtual bool UpdateLoadedProjectFileToCurrent(const TArray<FString>* StartupModuleNames, const FString& EngineIdentifier, FText& OutFailReason) OVERRIDE;
 	virtual bool SignSampleProject(const FString& FilePath, const FString& Category, FText& OutFailReason) OVERRIDE;
-	virtual bool QueryStatusForProject(const FString& FilePath, FProjectStatus& OutProjectStatus) const OVERRIDE;
+	virtual bool QueryStatusForProject(const FString& FilePath, const FString& EngineIdentifier, FProjectStatus& OutProjectStatus) const OVERRIDE;
 
 private:
 	/** The project that is currently loaded in the editor */
diff --git a/Engine/Source/Runtime/Projects/Public/Interfaces/IProjectManager.h b/Engine/Source/Runtime/Projects/Public/Interfaces/IProjectManager.h
--- a/Engine/Source/Runtime/Projects/Public/Interfaces/IProjectManager.h
+++ b/Engine/Source/Runtime/Projects/Public/Interfaces/IProjectManager.h
@@ -88,32 +88,35 @@ public:
 	 * 
 	 * @param NewProjectFilename	The filename of the file to be written out
 	 * @param StartupModuleNames	The list of modules to be loaded at startup
+	 * @param EngineIdentifier		Identifier of the engine that the project should be associated with
 	 * @param OutFailReason			When returning false, this provides a display reason why the file could not be created.
 	 *
 	 * @return true when the file was successfully written
 	 */
-	virtual bool GenerateNewProjectFile(const FString& NewProjectFilename, const TArray<FString>& StartupModuleNames, FText& OutFailReason) = 0;
+	virtual bool GenerateNewProjectFile(const FString& NewProjectFilename, const TArray<FString>& StartupModuleNames, const FString& EngineIdentifier, FText& OutFailReason) = 0;
 
 	/**
 	 * Generates a new project file based on a source file and saves it to disk at the specified location
 	 * 
 	 * @param SourceProjectFilename	The filename of the file to duplicate
 	 * @param NewProjectFilename	The filename of the file to be written out
+	 * @param EngineIdentifier		Identifier of the engine that the project should be associated with
 	 * @param OutFailReason			When returning false, this provides a display reason why the file could not be created.
 	 *
 	 * @return true when the file was successfully written
 	 */
-	virtual bool DuplicateProjectFile(const FString& SourceProjectFilename, const FString& NewProjectFilename, FText& OutFailReason) = 0;
+	virtual bool DuplicateProjectFile(const FString& SourceProjectFilename, const FString& NewProjectFilename, const FString& EngineIdentifier, FText& OutFailReason) = 0;
 
 	/**
 	 * Updates the loaded game project file to the current version and saves it to disk
 	 * 
 	 * @param StartupModuleNames	When specified, replaces the existing module names with these new ones			
+	 * @param EngineIdentifier		Identifier for the current engine
 	 * @param OutFailReason			When returning false, this provides a display reason why the file could not be created.
 	 *
 	 * @return true if the file was successfully updated and saved to disk
 	 */
-	virtual bool UpdateLoadedProjectFileToCurrent(const TArray<FString>* StartupModuleNames, FText& OutFailReason) = 0;
+	virtual bool UpdateLoadedProjectFileToCurrent(const TArray<FString>* StartupModuleNames, const FString& EngineIdentifier, FText& OutFailReason) = 0;
 
 	/**
 	 * Sets the project's EpicSampleNameHash (based on its filename) and category, then saves the file to disk.
@@ -132,11 +135,12 @@ public:
 	 * Gets status about the specified project
 	 *
 	 * @param FilePath				The filepath where the project is stored.
+	 * @param EngineIdentifier		Identifier for the current engine installation
 	 * @param OutProjectStatus		The status for the project.
 	 *
 	 * @return	 true if the file was successfully open and read
 	 */
-	virtual bool QueryStatusForProject(const FString& FilePath, FProjectStatus& OutProjectStatus) const = 0;
+	virtual bool QueryStatusForProject(const FString& FilePath, const FString& EngineIdentifier, FProjectStatus& OutProjectStatus) const = 0;
 
 	/** Helper functions to reduce the syntax complexity of commonly used functions */
 	static const FString& GetProjectFileExtension() { return Get().NonStaticGetProjectFileExtension(); }
