#ue4
#ios
* add extra messaging which links to documentation when uproject is not found (IOS-only)

TTP330364

diff --git a/Engine/Source/Runtime/Projects/Private/ProjectManager.cpp b/Engine/Source/Runtime/Projects/Private/ProjectManager.cpp
--- a/Engine/Source/Runtime/Projects/Private/ProjectManager.cpp
+++ b/Engine/Source/Runtime/Projects/Private/ProjectManager.cpp
@@ -104,8 +104,13 @@ bool FProjectManager::LoadProjectFile( const FString& InProjectFile )
 		}
 	}
 	
+#if PLATFORM_IOS
+    FString UpdatedMessage = FString::Printf(TEXT("%s\n%s"), *FailureReason.ToString(), TEXT("For troubleshooting, please go to https://docs.unrealengine.com/latest/INT/Platforms/iOS/GettingStarted/index.html"));
+    FailureReason = FText::FromString(UpdatedMessage);
+#endif
 	UE_LOG(LogProjectManager, Error, TEXT("%s"), *FailureReason.ToString());
 	FMessageDialog::Open(EAppMsgType::Ok, FailureReason);
+    
 	return false;
 }
 
