Add methods to determine whether a project or plugin's modules are up-to-date.

diff --git a/Engine/Source/Runtime/Core/Private/Modules/ModuleManager.cpp b/Engine/Source/Runtime/Core/Private/Modules/ModuleManager.cpp
--- a/Engine/Source/Runtime/Core/Private/Modules/ModuleManager.cpp
+++ b/Engine/Source/Runtime/Core/Private/Modules/ModuleManager.cpp
@@ -123,6 +123,19 @@ bool FModuleManager::IsModuleLoaded( const FName InModuleName ) const
 	return false;
 }
 
+bool FModuleManager::IsModuleUpToDate( const FName InModuleName ) const
+{
+	TMap<FName, FString> ModulePathMap;
+	FindModulePaths(*InModuleName.ToString(), ModulePathMap);
+
+	if(ModulePathMap.Num() != 1)
+	{
+		return false;
+	}
+
+	return CheckModuleCompatibility(*TMap<FName, FString>::TConstIterator(ModulePathMap).Value());
+}
+
 void FModuleManager::AddModule( const FName InModuleName )
 {
 	// Do we already know about this module?  If not, we'll create information for this module now.
@@ -1139,7 +1152,7 @@ bool FModuleManager::GetModuleFileTimeStamp(TSharedRef<const FModuleInfo> Module
 	return false;
 }
 
-void FModuleManager::FindModulePaths(const TCHAR* NamePattern, TMap<FName, FString> &OutModulePaths)
+void FModuleManager::FindModulePaths(const TCHAR* NamePattern, TMap<FName, FString> &OutModulePaths) const
 {
 	// Search through the engine directory
 	FindModulePathsInDirectory(FPlatformProcess::GetModulesDirectory(), false, NamePattern, OutModulePaths);
@@ -1157,7 +1170,7 @@ void FModuleManager::FindModulePaths(const TCHAR* NamePattern, TMap<FName, FStri
 	}
 }
 
-void FModuleManager::FindModulePathsInDirectory(const FString& InDirectoryName, bool bIsGameDirectory, const TCHAR* NamePattern, TMap<FName, FString> &OutModulePaths)
+void FModuleManager::FindModulePathsInDirectory(const FString& InDirectoryName, bool bIsGameDirectory, const TCHAR* NamePattern, TMap<FName, FString> &OutModulePaths) const
 {
 	// Get the module configuration for this directory type
 	#if UE_BUILD_DEBUG
@@ -1492,6 +1505,17 @@ void FModuleManager::MakeUniqueModuleFilename( const FName InModuleName, FString
 	while (IFileManager::Get().GetFileAgeSeconds(*UniqueModuleFileName) != -1.0);
 }
 
+const TCHAR *FModuleManager::GetUBTConfiguration()
+{
+#if UE_BUILD_DEBUG
+	return TEXT("Debug");
+#elif UE_BUILD_SHIPPING
+	return TEXT("Shipping");
+#else
+	static bool bIsDebugGame = FParse::Param(FCommandLine::Get(), TEXT("debug"));
+	return bIsDebugGame? TEXT("DebugGame") : TEXT("Development");
+#endif
+}
 
 bool FModuleManager::StartCompilingModuleDLLs(const FString& GameName, const TArray< FModuleToRecompile >& ModuleNames, 
 	const FRecompileModulesCallback& InRecompileModulesCallback, FOutputDevice& Ar, bool bInFailIfGeneratedCodeChanges, 
@@ -1503,15 +1527,7 @@ bool FModuleManager::StartCompilingModuleDLLs(const FString& GameName, const TAr
 	ModulesThatWereBeingRecompiled = ModulesBeingCompiled;
 
 	const TCHAR* BuildPlatformName = FPlatformMisc::GetUBTPlatform();
-
-	FString BuildConfigurationName = 
-#if UE_BUILD_DEBUG
-		TEXT( "Debug");
-#elif UE_BUILD_SHIPPING
-		TEXT( "Shipping" );
-#else
-		FParse::Param(FCommandLine::Get(), TEXT("debug"))? TEXT( "DebugGame" ) : TEXT( "Development" );
-#endif
+	const TCHAR* BuildConfigurationName = GetUBTConfiguration();
 
 	RecompileModulesCallback = InRecompileModulesCallback;
 
@@ -1554,7 +1570,7 @@ bool FModuleManager::StartCompilingModuleDLLs(const FString& GameName, const TAr
 
 	FString CmdLineParams = FString::Printf( TEXT( "%s%s %s %s %s%s" ), 
 		*GameName, *ModuleArg, 
-		BuildPlatformName, *BuildConfigurationName, 
+		BuildPlatformName, BuildConfigurationName, 
 		*ExtraArg, *InAdditionalCmdLineArgs );
 
 	const bool bInvocationSuccessful = InvokeUnrealBuildTool(CmdLineParams, Ar);
diff --git a/Engine/Source/Runtime/Core/Public/Modules/ModuleManager.h b/Engine/Source/Runtime/Core/Public/Modules/ModuleManager.h
--- a/Engine/Source/Runtime/Core/Public/Modules/ModuleManager.h
+++ b/Engine/Source/Runtime/Core/Public/Modules/ModuleManager.h
@@ -428,6 +428,14 @@ public:
 	 */
 	void SetGameBinariesDirectory(const TCHAR* InDirectory);
 
+	/**
+	 * Checks to see if the specified module exists and is compatible with the current engine version. 
+	 *
+	 * @param	InModuleName		The base name of the module file.
+	 *
+	 * @return	True if module exists and is up to date.
+	 */
+	bool IsModuleUpToDate( const FName InModuleName ) const;
 
 	/**
 	 * Adds a module to our list of modules, unless it's already known
@@ -447,6 +455,13 @@ public:
 	 */
 	bool DoesLoadedModuleHaveUObjects( const FName ModuleName );
 
+	/**
+	 * Gets the build configuration for compiling modules, as required by UBT.
+	 *
+	 * @return	Configuration name for UBT.
+	 */
+	static const TCHAR *GetUBTConfiguration( );
+
 private:
 
 	/** Private constructor.  Singleton instance is always constructed on demand. */
@@ -597,9 +612,11 @@ private:
 	virtual bool Exec( UWorld* Inworld, const TCHAR* Cmd, FOutputDevice& Ar ) OVERRIDE;
 	// End of FExec interface.
 
+public:
 	/** @returns Static: Returns arguments to pass to UnrealBuildTool when compiling modules */
 	static FString MakeUBTArgumentsForModuleCompiling();
 
+private:
 	/** Called during CheckForFinishedModuleDLLCompile() for each successfully recomplied module */
 	void OnModuleCompileSucceeded(FName ModuleName, TSharedRef<FModuleInfo> ModuleInfo);
 
@@ -616,10 +633,10 @@ private:
 	static bool GetModuleFileTimeStamp(TSharedRef<const FModuleInfo> ModuleInfo, FDateTime& OutFileTimeStamp);
 
 	/** Finds modules matching a given name wildcard. */
-	void FindModulePaths(const TCHAR *NamePattern, TMap<FName, FString> &OutModulePaths);
+	void FindModulePaths(const TCHAR *NamePattern, TMap<FName, FString> &OutModulePaths) const;
 
 	/** Finds modules matching a given name wildcard within a given directory. */
-	void FindModulePathsInDirectory(const FString &DirectoryName, bool bIsGameDirectory, const TCHAR *NamePattern, TMap<FName, FString> &OutModulePaths);
+	void FindModulePathsInDirectory(const FString &DirectoryName, bool bIsGameDirectory, const TCHAR *NamePattern, TMap<FName, FString> &OutModulePaths) const;
 
 private:
 
diff --git a/Engine/Source/Runtime/Projects/Private/PluginManager.cpp b/Engine/Source/Runtime/Projects/Private/PluginManager.cpp
--- a/Engine/Source/Runtime/Projects/Private/PluginManager.cpp
+++ b/Engine/Source/Runtime/Projects/Private/PluginManager.cpp
@@ -363,6 +363,27 @@ bool FPluginManager::IsPluginModule( const FName ModuleName ) const
 	return ( ModulePluginMap.Find( ModuleName ) != NULL );
 }
 
+bool FPluginManager::AreEnabledPluginModulesUpToDate()
+{
+	if (!bEarliestPhaseProcessed)
+	{
+		EnablePluginsThatAreConfiguredToBeEnabled();
+		RegisterEnabledPluginMountPoints();
+		bEarliestPhaseProcessed = true;
+	}
+
+	for (TArray< TSharedRef< FPlugin > >::TConstIterator Iter(AllPlugins); Iter; ++Iter)
+	{
+		const TSharedRef< FPlugin > &Plugin = *Iter;
+		if (Plugin->IsEnabled() && !Plugin->AreModulesUpToDate())
+		{
+			return false;
+		}
+	}
+
+	return true;
+}
+
 void FPluginManager::GetPluginsConfiguredToBeEnabled(TArray<FString>& OutPluginNames) const
 {
 #if !IS_PROGRAM
diff --git a/Engine/Source/Runtime/Projects/Private/PluginManager.h b/Engine/Source/Runtime/Projects/Private/PluginManager.h
--- a/Engine/Source/Runtime/Projects/Private/PluginManager.h
+++ b/Engine/Source/Runtime/Projects/Private/PluginManager.h
@@ -125,6 +125,7 @@ public:
 	virtual void LoadModulesForEnabledPlugins( const ELoadingPhase::Type LoadingPhase ) OVERRIDE;
 	virtual void SetRegisterMountPointDelegate( const FRegisterMountPointDelegate& Delegate ) OVERRIDE;
 	virtual bool IsPluginModule( const FName ModuleName ) const OVERRIDE;
+	virtual bool AreEnabledPluginModulesUpToDate() OVERRIDE;
 	virtual TArray< FPluginStatus > QueryStatusForAllPlugins() const OVERRIDE;
 	virtual void SetPluginEnabled( const FString& PluginName, bool bEnabled ) OVERRIDE;
 	virtual bool IsRestartRequired() const OVERRIDE;
diff --git a/Engine/Source/Runtime/Projects/Private/PluginManagerShared.cpp b/Engine/Source/Runtime/Projects/Private/PluginManagerShared.cpp
--- a/Engine/Source/Runtime/Projects/Private/PluginManagerShared.cpp
+++ b/Engine/Source/Runtime/Projects/Private/PluginManagerShared.cpp
@@ -11,9 +11,6 @@ void FProjectOrPlugin::LoadModules( const ELoadingPhase::Type LoadingPhase, TMap
 {
 	const FProjectOrPluginInfo& ProjectOrPluginInfo = GetProjectOrPluginInfo();
 
-	// cache the string for the current platform
-	static FString UBTPlatform(FPlatformMisc::GetUBTPlatform());
-
 	for( TArray< FProjectOrPluginInfo::FModuleInfo >::TConstIterator ModuleInfoIt( ProjectOrPluginInfo.Modules.CreateConstIterator() ); ModuleInfoIt; ++ModuleInfoIt )
 	{
 		const FProjectOrPluginInfo::FModuleInfo& ModuleInfo = *ModuleInfoIt;
@@ -21,96 +18,22 @@ void FProjectOrPlugin::LoadModules( const ELoadingPhase::Type LoadingPhase, TMap
 		// Don't need to do anything if this module is already loaded
 		if( !FModuleManager::Get().IsModuleLoaded( ModuleInfo.Name ) )
 		{
-			bool bShouldLoadModule = false;
-			// Only load modules that support this platform
-			if ((ModuleInfo.WhitelistPlatforms.Num() == 0 || ModuleInfo.WhitelistPlatforms.Contains(UBTPlatform)) &&
-				(ModuleInfo.BlacklistPlatforms.Num() == 0 || !ModuleInfo.BlacklistPlatforms.Contains(UBTPlatform)))
+			if( LoadingPhase == ModuleInfo.LoadingPhase && ShouldLoadModule(ModuleInfo) )
 			{
-				bShouldLoadModule = true;
-			}
+				// @todo plugin: DLL search problems.  Plugins that statically depend on other modules within this plugin may not be found?  Need to test this.
 
-			// Only load modules for the current loading phase
-			if( bShouldLoadModule && LoadingPhase == ModuleInfo.LoadingPhase )
-			{
-				bShouldLoadModule = false;
-				switch( ModuleInfo.Type )
+				// NOTE: Loading this module may cause other modules to become loaded, both in the engine or game, or other modules 
+				//       that are part of this project or plugin.  That's totally fine.
+				ELoadModuleFailureReason::Type FailureReason;
+				const TSharedPtr<IModuleInterface>& ModuleInterface = FModuleManager::Get().LoadModuleWithFailureReason( ModuleInfo.Name, FailureReason );
+				if( ModuleInterface.IsValid() )
 				{
-				case FProjectOrPluginInfo::EModuleType::Runtime:
-					{
-#if WITH_ENGINE || WITH_PLUGIN_SUPPORT
-						bShouldLoadModule = true;
-#endif		// WITH_ENGINE
-					}
-					break;
-
-				case FProjectOrPluginInfo::EModuleType::RuntimeNoCommandlet:
-					{
-#if WITH_ENGINE || WITH_PLUGIN_SUPPORT
-						if (!IsRunningCommandlet())
-						{
-							bShouldLoadModule = true;
-						}
-#endif		// WITH_ENGINE
-					}
-					break;
-
-				case FProjectOrPluginInfo::EModuleType::Developer:
-					{
-#if WITH_UNREAL_DEVELOPER_TOOLS
-						bShouldLoadModule = true;
-#endif		// WITH_UNREAL_DEVELOPER_TOOLS
-					}
-					break;
-
-				case FProjectOrPluginInfo::EModuleType::Editor:
-					{
-#if WITH_EDITOR
-						if( GIsEditor )
-						{
-							bShouldLoadModule = true;
-						}
-#endif		// WITH_EDITOR
-					}
-					break;
-					
-				case FProjectOrPluginInfo::EModuleType::EditorNoCommandlet:
-					{
-#if WITH_EDITOR
-						if (GIsEditor && !IsRunningCommandlet())
-						{
-							bShouldLoadModule = true;
-						}
-#endif		// WITH_EDITOR
-					}
-					break;
-
-				case FProjectOrPluginInfo::EModuleType::Program:
-					{
-#if WITH_PLUGIN_SUPPORT && IS_PROGRAM
-						bShouldLoadModule = true;
-#endif		// WITH_PLUGIN_SUPPORT
-					}
-					break;
+					// Module loaded OK (or was already loaded.)
 				}
-
-
-				if( bShouldLoadModule )
+				else 
 				{
-					// @todo plugin: DLL search problems.  Plugins that statically depend on other modules within this plugin may not be found?  Need to test this.
-
-					// NOTE: Loading this module may cause other modules to become loaded, both in the engine or game, or other modules 
-					//       that are part of this project or plugin.  That's totally fine.
-					ELoadModuleFailureReason::Type FailureReason;
-					const TSharedPtr<IModuleInterface>& ModuleInterface = FModuleManager::Get().LoadModuleWithFailureReason( ModuleInfo.Name, FailureReason );
-					if( ModuleInterface.IsValid() )
-					{
-						// Module loaded OK (or was already loaded.)
-					}
-					else 
-					{
-						// The module failed to load. Note this in the ModuleLoadErrors list.
-						ModuleLoadErrors.Add(ModuleInfo.Name, FailureReason);
-					}
+					// The module failed to load. Note this in the ModuleLoadErrors list.
+					ModuleLoadErrors.Add(ModuleInfo.Name, FailureReason);
 				}
 			}
 		}
@@ -321,6 +244,20 @@ bool FProjectOrPlugin::IsUpToDate( const FString &EngineIdentifier ) const
 	return true;
 }
 
+bool FProjectOrPlugin::AreModulesUpToDate( ) const
+{
+	const FProjectOrPluginInfo& ProjectOrPluginInfo = GetProjectOrPluginInfo();
+	for (TArray< FProjectOrPluginInfo::FModuleInfo >::TConstIterator Iter(ProjectOrPluginInfo.Modules); Iter; ++Iter)
+	{
+		const FProjectOrPluginInfo::FModuleInfo &ModuleInfo = *Iter;
+		if (ShouldBuildModule(ModuleInfo) && !FModuleManager::Get().IsModuleUpToDate(ModuleInfo.Name))
+		{
+			return false;
+		}
+	}
+	return true;
+}
+
 void FProjectOrPlugin::UpdateVersionToCurrent( const FString &EngineIdentifier )
 {
 	FProjectOrPluginInfo& ProjectOrPluginInfo = GetProjectOrPluginInfo();
@@ -348,6 +285,104 @@ void FProjectOrPlugin::ReplaceModulesInProject(const TArray<FString>* StartupMod
 	}
 }
 
+bool FProjectOrPlugin::ShouldBuildModule(const FProjectOrPluginInfo::FModuleInfo& ModuleInfo) const
+{
+	// Cache the string for the current platform
+	static FString UBTPlatform(FPlatformMisc::GetUBTPlatform());
+
+	// Check the platform is whitelisted
+	if(ModuleInfo.WhitelistPlatforms.Num() > 0 && !ModuleInfo.WhitelistPlatforms.Contains(UBTPlatform))
+	{
+		return false;
+	}
+
+	// Check the platform is not blacklisted
+	if(ModuleInfo.BlacklistPlatforms.Num() > 0 && ModuleInfo.BlacklistPlatforms.Contains(UBTPlatform))
+	{
+		return false;
+	}
+
+	// Check the module is compatible with this target. This should match UEBuildTarget.ShouldIncludePluginModule in UBT
+	switch (ModuleInfo.Type)
+	{
+	case FProjectOrPluginInfo::EModuleType::Runtime:
+	case FProjectOrPluginInfo::EModuleType::RuntimeNoCommandlet:
+		return true;
+
+	case FProjectOrPluginInfo::EModuleType::Developer:
+		#if WITH_UNREAL_DEVELOPER_TOOLS
+			return true;
+		#endif
+		break;
+
+	case FProjectOrPluginInfo::EModuleType::Editor:
+	case FProjectOrPluginInfo::EModuleType::EditorNoCommandlet:
+		#if WITH_EDITOR
+			return true;
+		#endif
+		break;
+
+	case FProjectOrPluginInfo::EModuleType::Program:
+		#if IS_PROGRAM
+			return true;
+		#endif
+		break;
+	}
+
+	return false;
+}
+
+bool FProjectOrPlugin::ShouldLoadModule(const FProjectOrPluginInfo::FModuleInfo& ModuleInfo) const
+{
+	// Check that the module is built for this configuration
+	if(!ShouldBuildModule(ModuleInfo))
+	{
+		return false;
+	}
+
+	// Check that the runtime environment allows it to be loaded
+	switch (ModuleInfo.Type)
+	{
+	case FProjectOrPluginInfo::EModuleType::Runtime:
+		#if WITH_ENGINE || WITH_PLUGIN_SUPPORT
+			return true;
+		#endif
+		break;
+
+	case FProjectOrPluginInfo::EModuleType::RuntimeNoCommandlet:
+		#if WITH_ENGINE || WITH_PLUGIN_SUPPORT
+			if(!IsRunningCommandlet()) return true;
+		#endif
+		break;
+
+	case FProjectOrPluginInfo::EModuleType::Developer:
+		#if WITH_UNREAL_DEVELOPER_TOOLS
+			return true;
+		#endif
+		break;
+
+	case FProjectOrPluginInfo::EModuleType::Editor:
+		#if WITH_EDITOR
+			if(GIsEditor) return true;
+		#endif
+		break;
+
+	case FProjectOrPluginInfo::EModuleType::EditorNoCommandlet:
+		#if WITH_EDITOR
+			if(GIsEditor && !IsRunningCommandlet()) return true;
+		#endif
+		break;
+
+	case FProjectOrPluginInfo::EModuleType::Program:
+		#if WITH_PLUGIN_SUPPORT && IS_PROGRAM
+			return true;
+		#endif
+		break;
+	}
+
+	return false;
+}
+
 bool FProjectOrPlugin::ReadNumberFromJSON(const TSharedRef< FJsonObject >& FileObject, const FString& PropertyName, int32& OutNumber ) const
 {
 	int64 LargeNumber;
diff --git a/Engine/Source/Runtime/Projects/Private/PluginManagerShared.h b/Engine/Source/Runtime/Projects/Private/PluginManagerShared.h
--- a/Engine/Source/Runtime/Projects/Private/PluginManagerShared.h
+++ b/Engine/Source/Runtime/Projects/Private/PluginManagerShared.h
@@ -228,6 +228,13 @@ public:
 	 */
 	bool IsUpToDate( const FString &EngineIdentifier ) const;
 
+	/**
+	 * Checks whether the modules for a project or plugin need to be compiled
+	 *
+	 * @return true if the project or plugin needs to be compiled
+	 */
+	bool AreModulesUpToDate( ) const;
+
 	/**
 	 * Updates all version info to match the currently running executable.
 	 */
@@ -241,6 +248,12 @@ public:
 	void ReplaceModulesInProject(const TArray<FString>* StartupModuleNames);
 
 protected:
+	/** Checks whether a given module should be built for this build configuration */
+	bool ShouldBuildModule(const FProjectOrPluginInfo::FModuleInfo& ModuleInfo) const;
+
+	/** Checks whether a given module should be loaded for this build configuration */
+	bool ShouldLoadModule(const FProjectOrPluginInfo::FModuleInfo& ModuleInfo) const;
+
 	/** @return Exposes access to the project or plugin's descriptor */
 	virtual FProjectOrPluginInfo& GetProjectOrPluginInfo() = 0;
 	virtual const FProjectOrPluginInfo& GetProjectOrPluginInfo() const = 0;
diff --git a/Engine/Source/Runtime/Projects/Private/ProjectManager.cpp b/Engine/Source/Runtime/Projects/Private/ProjectManager.cpp
--- a/Engine/Source/Runtime/Projects/Private/ProjectManager.cpp
+++ b/Engine/Source/Runtime/Projects/Private/ProjectManager.cpp
@@ -172,6 +172,11 @@ bool FProjectManager::LoadModulesForProject( const ELoadingPhase::Type LoadingPh
 	return bSuccess;
 }
 
+bool FProjectManager::AreProjectModulesUpToDate()
+{
+	return !CurrentlyLoadedProject.IsValid() || CurrentlyLoadedProject->AreModulesUpToDate();
+}
+
 const FString& FProjectManager::GetAutoLoadProjectFileName()
 {
 	static FString RecentProjectFileName = FPaths::Combine(*FPaths::GameAgnosticSavedDir(), TEXT("AutoLoadProject.txt"));
diff --git a/Engine/Source/Runtime/Projects/Private/ProjectManager.h b/Engine/Source/Runtime/Projects/Private/ProjectManager.h
--- a/Engine/Source/Runtime/Projects/Private/ProjectManager.h
+++ b/Engine/Source/Runtime/Projects/Private/ProjectManager.h
@@ -86,6 +86,7 @@ public:
 	/** IProjectManager interface */
 	virtual bool LoadProjectFile( const FString& ProjectFile ) OVERRIDE;
 	virtual bool LoadModulesForProject( const ELoadingPhase::Type LoadingPhase ) OVERRIDE;
+	virtual bool AreProjectModulesUpToDate( ) OVERRIDE;
 	virtual const FString& GetAutoLoadProjectFileName() OVERRIDE;
 	virtual const FString& NonStaticGetProjectFileExtension() OVERRIDE;
 	virtual bool GenerateNewProjectFile(const FString& NewProjectFilename, const TArray<FString>& ModuleNames, const FString& EngineIdentifier, FText& OutFailReason) OVERRIDE;
diff --git a/Engine/Source/Runtime/Projects/Public/Interfaces/IPluginManager.h b/Engine/Source/Runtime/Projects/Public/Interfaces/IPluginManager.h
--- a/Engine/Source/Runtime/Projects/Public/Interfaces/IPluginManager.h
+++ b/Engine/Source/Runtime/Projects/Public/Interfaces/IPluginManager.h
@@ -97,6 +97,12 @@ public:
 	 */
 	virtual bool IsPluginModule( const FName ModuleName ) const = 0;
 
+	/** 
+	 * Checks whether modules for the enabled plugins are up to date
+	 *
+	 * @returns true if the enabled plugin modules are up to date
+	 */
+	virtual bool AreEnabledPluginModulesUpToDate( ) = 0;
 
 	/**
 	 * Gets status about all currently known plugins
diff --git a/Engine/Source/Runtime/Projects/Public/Interfaces/IProjectManager.h b/Engine/Source/Runtime/Projects/Public/Interfaces/IProjectManager.h
--- a/Engine/Source/Runtime/Projects/Public/Interfaces/IProjectManager.h
+++ b/Engine/Source/Runtime/Projects/Public/Interfaces/IProjectManager.h
@@ -67,6 +67,13 @@ public:
 	 */
 	virtual bool LoadModulesForProject( const ELoadingPhase::Type LoadingPhase ) = 0;
 
+	/**
+	 * Checks if the modules for a project are up to date
+	 *
+	 * @return	true if the UBT needs to be run to recompile modules for a project.
+	 */
+	virtual bool AreProjectModulesUpToDate( ) = 0;
+
 	/**
 	 * Gets the name of the text file that contains the most recently loaded filename.
 	 *
