fix for Memory Card, Black Jack, and Swing Ninja not launching properly after being packaged
TTP334744
#ue4

diff --git a/Engine/Source/Runtime/Projects/Private/ProjectManager.cpp b/Engine/Source/Runtime/Projects/Private/ProjectManager.cpp
--- a/Engine/Source/Runtime/Projects/Private/ProjectManager.cpp
+++ b/Engine/Source/Runtime/Projects/Private/ProjectManager.cpp
@@ -40,19 +40,6 @@ bool FProject::VerifyProjectCanBeLoaded( bool bPromptIfSavedWithNewerVersionOfEn
 		}
 	}
 
-	// Older projects did not specify a module list. In these projects, the only assumed module was one named after the game, loaded in the default phase.
-	if ( ProjectInfo.FileVersion < VER_PROJECT_FILE_PLUGIN_UNIFICATION )
-	{
-		// Assuming that a module is needed to be loaded if a source folder exists
-		const FString BinariesFolder = FPaths::Combine( *FPaths::GetPath(ProjectInfo.LoadedFromFile), TEXT("Binaries"), FPlatformProcess::GetBinariesSubdirectory() );
-		if ( IFileManager::Get().DirectoryExists(*BinariesFolder) )
-		{
-			FProjectOrPluginInfo::FModuleInfo MainGameModule;
-			MainGameModule.Name = FName(*ProjectInfo.Name);
-			ProjectInfo.Modules.Add(MainGameModule);
-		}
-	}
-
 	return true;
 }
 
