add check for third-party plugins for IOS on PC as we don't support that currently in the released version
#ue4
#ios

diff --git a/Engine/Source/Editor/MainFrame/Private/Frame/MainFrameActions.cpp b/Engine/Source/Editor/MainFrame/Private/Frame/MainFrameActions.cpp
--- a/Engine/Source/Editor/MainFrame/Private/Frame/MainFrameActions.cpp
+++ b/Engine/Source/Editor/MainFrame/Private/Frame/MainFrameActions.cpp
@@ -505,6 +505,11 @@ void FMainFrameActionCallbacks::PackageProject( const FName InPlatformInfoName )
 		FMessageDialog::Open(EAppMsgType::Ok, LOCTEXT("PackagingNotWorkingMessage", "Sorry, packaging is currently not supported for code-based iOS projects. This feature will be available in a future release.") );
 		return;
 	}
+	if (FRocketSupport::IsRocket() && PlatformInfo->TargetPlatformName == FName("IOS") && IPluginManager::Get().HasThirdPartyPlugin())
+	{
+		FMessageDialog::Open(EAppMsgType::Ok, LOCTEXT("PackagingNotWorkingMessage", "Sorry, packaging is currently not supported for content based projects with third-party plugins. This feature will be available in a future release.") );
+		return;
+	}
 #endif
 
 	if (!FModuleManager::LoadModuleChecked<IProjectTargetPlatformEditorModule>("ProjectTargetPlatformEditor").ShowUnsupportedTargetWarning(PlatformInfo->VanillaPlatformName))
diff --git a/Engine/Source/Editor/UnrealEd/Private/PlayLevel.cpp b/Engine/Source/Editor/UnrealEd/Private/PlayLevel.cpp
--- a/Engine/Source/Editor/UnrealEd/Private/PlayLevel.cpp
+++ b/Engine/Source/Editor/UnrealEd/Private/PlayLevel.cpp
@@ -1210,6 +1210,12 @@ void UEditorEngine::PlayUsingLauncher()
 			bPlayUsingLauncher = false;
 			return;
 		}
+		if (FRocketSupport::IsRocket() && PlayUsingLauncherDeviceId.Left(PlayUsingLauncherDeviceId.Find(TEXT("@"))) == TEXT("IOS") && IPluginManager::Get().HasThirdPartyPlugin())
+		{
+			FMessageDialog::Open(EAppMsgType::Ok, LOCTEXT("LaunchOnDeviceNotWorkingMessage", "Sorry, launching on a device is currently not supported for content based projects with third-party plugins. This feature will be available in a future release.") );
+			bPlayUsingLauncher = false;
+			return;
+		}
 #endif
 
 		ILauncherProfileRef LauncherProfile = LauncherServicesModule.CreateProfile(TEXT("Play On Device"));
diff --git a/Engine/Source/Runtime/Projects/Private/PluginManager.cpp b/Engine/Source/Runtime/Projects/Private/PluginManager.cpp
--- a/Engine/Source/Runtime/Projects/Private/PluginManager.cpp
+++ b/Engine/Source/Runtime/Projects/Private/PluginManager.cpp
@@ -389,6 +389,51 @@ void FPluginManager::GetPluginsConfiguredToBeEnabled(TArray<FString>& OutPluginN
 #endif
 }
 
+bool FPluginManager::HasThirdPartyPlugin() const
+{
+	static bool bInit = false;
+	TArray<FString> StandardPlugins;
+	if (!bInit)
+	{
+		StandardPlugins.Add(TEXT("BlankPlugin"));
+		StandardPlugins.Add(TEXT("UObjectPlugin"));
+		StandardPlugins.Add(TEXT("IntroTutorials"));
+		StandardPlugins.Add(TEXT("PluginsEditor"));
+		StandardPlugins.Add(TEXT("EpicSurvey"));
+		StandardPlugins.Add(TEXT("OculusRift"));
+		StandardPlugins.Add(TEXT("CustomMeshComponent"));
+		StandardPlugins.Add(TEXT("MessagingDebugger"));
+		StandardPlugins.Add(TEXT("PerforceSourceControl"));
+		StandardPlugins.Add(TEXT("SubversionSourceControl"));
+		StandardPlugins.Add(TEXT("UdpMessaging"));
+		StandardPlugins.Add(TEXT("WindowsMoviePlayer"));
+		StandardPlugins.Add(TEXT("CableComponent"));
+		StandardPlugins.Add(TEXT("ExampleDeviceProfileSelector"));
+		StandardPlugins.Add(TEXT("WinDualShock"));
+		StandardPlugins.Add(TEXT("VisualStudioSourceCodeAccess"));
+		StandardPlugins.Add(TEXT("XCodeSourceCodeAccess"));
+		StandardPlugins.Add(TEXT("SlateRemote"));
+		StandardPlugins.Add(TEXT("ScriptPlugin"));
+		StandardPlugins.Add(TEXT("ScriptEditorPlugin"));
+		StandardPlugins.Add(TEXT("SpeedTreeImporter"));
+		StandardPlugins.Add(TEXT("AppleMoviePlayer"));
+		StandardPlugins.Add(TEXT("IOSDeviceProfileSelector"));
+		bInit = true;
+	}
+
+	TArray<FString> EnabledPlugins;
+	GConfig->GetArray( TEXT("Plugins"), TEXT("EnabledPlugins"), EnabledPlugins, GEngineIni );
+
+	for (int Index = 0; Index < EnabledPlugins.Num(); ++Index)
+	{
+		if (!StandardPlugins.Contains(EnabledPlugins[Index]))
+		{
+			return true;
+		}
+	}
+	return false;
+}
+
 void FPluginManager::ConfigurePluginToBeEnabled(const FString& PluginName, bool bEnabled)
 {
 	TArray< FString > EnabledPluginNames;
diff --git a/Engine/Source/Runtime/Projects/Private/PluginManager.h b/Engine/Source/Runtime/Projects/Private/PluginManager.h
--- a/Engine/Source/Runtime/Projects/Private/PluginManager.h
+++ b/Engine/Source/Runtime/Projects/Private/PluginManager.h
@@ -129,6 +129,7 @@ public:
 	virtual TArray< FPluginStatus > QueryStatusForAllPlugins() const override;
 	virtual void SetPluginEnabled( const FString& PluginName, bool bEnabled ) override;
 	virtual bool IsRestartRequired() const override;
+	virtual bool HasThirdPartyPlugin() const override;
 
 private:
 
diff --git a/Engine/Source/Runtime/Projects/Public/Interfaces/IPluginManager.h b/Engine/Source/Runtime/Projects/Public/Interfaces/IPluginManager.h
--- a/Engine/Source/Runtime/Projects/Public/Interfaces/IPluginManager.h
+++ b/Engine/Source/Runtime/Projects/Public/Interfaces/IPluginManager.h
@@ -120,4 +120,11 @@ public:
 	 * @return true if plugin changes have been made that will only take effect after a restart
 	 */
 	virtual bool IsRestartRequired() const = 0;
+
+	/**
+	 * Determines if a third-party plugin has been requested in the executable
+	 *
+	 * @return true if a third party plugin is enabled
+	 */
+	virtual bool HasThirdPartyPlugin() const = 0;
 };
