Make a standalone class for module descriptors in .uproject and .uplugin files.

diff --git a/Engine/Source/Runtime/Projects/Private/ModuleDescriptor.cpp b/Engine/Source/Runtime/Projects/Private/ModuleDescriptor.cpp
new file mode 100644
--- /dev/null
+++ b/Engine/Source/Runtime/Projects/Private/ModuleDescriptor.cpp
@@ -0,0 +1,363 @@
+// Copyright 1998-2014 Epic Games, Inc. All Rights Reserved.
+
+#include "ProjectsPrivatePCH.h"
+
+#define LOCTEXT_NAMESPACE "PluginManagerShared"
+
+static int32 GetLoadOrderValue( ELoadingPhase::Type Phase )
+{
+	switch( Phase )
+	{
+	case ELoadingPhase::PostDefault:
+		return 500;
+
+	case ELoadingPhase::Default:
+		return 400;
+
+	case ELoadingPhase::PreDefault:
+		return 300;
+			
+	case ELoadingPhase::PreLoadingScreen:
+		return 200;
+
+	case ELoadingPhase::PostConfigInit:
+		return 100;
+
+	default:
+		ensureMsgf( false, TEXT( "Unrecognized ELoadingPhase value: %i" ), Phase );
+		return 0;
+	}
+}
+
+ELoadingPhase::Type ELoadingPhase::Earliest(ELoadingPhase::Type PhaseA, ELoadingPhase::Type PhaseB)
+{
+	const int32 PhaseAValue = GetLoadOrderValue(PhaseA);
+	const int32 PhaseBValue = GetLoadOrderValue(PhaseB);
+
+	return PhaseAValue < PhaseBValue ? PhaseA : PhaseB;
+}
+
+ELoadingPhase::Type ELoadingPhase::FromString( const TCHAR *String )
+{
+	ELoadingPhase::Type TestType = (ELoadingPhase::Type)0;
+	for(; TestType < ELoadingPhase::Max; TestType = (ELoadingPhase::Type)(TestType + 1))
+	{
+		const TCHAR *TestString = ToString(TestType);
+		if(FCString::Stricmp(String, TestString) == 0)
+		{
+			break;
+		}
+	}
+	return TestType;
+}
+
+const TCHAR* ELoadingPhase::ToString( const ELoadingPhase::Type Value )
+{
+	switch( Value )
+	{
+	case Default:
+		return TEXT( "Default" );
+
+	case PostDefault:
+		return TEXT( "PostDefault" );
+
+	case PreDefault:
+		return TEXT( "PreDefault" );
+
+	case PostConfigInit:
+		return TEXT( "PostConfigInit" );
+			
+	case PreLoadingScreen:
+		return TEXT( "PreLoadingScreen" );
+
+	default:
+		ensureMsgf( false, TEXT( "Unrecognized ELoadingPhase value: %i" ), Value );
+		return NULL;
+	}
+}
+
+EHostType::Type EHostType::FromString( const TCHAR *String )
+{
+	EHostType::Type TestType = (EHostType::Type)0;
+	for(; TestType < EHostType::Max; TestType = (EHostType::Type)(TestType + 1))
+	{
+		const TCHAR *TestString = ToString(TestType);
+		if(FCString::Stricmp(String, TestString) == 0)
+		{
+			break;
+		}
+	}
+	return TestType;
+}
+
+const TCHAR* EHostType::ToString( const EHostType::Type Value )
+{
+	switch( Value )
+	{
+		case Runtime:
+			return TEXT( "Runtime" );
+
+		case RuntimeNoCommandlet:
+			return TEXT( "RuntimeNoCommandlet" );
+
+		case Developer:
+			return TEXT( "Developer" );
+
+		case Editor:
+			return TEXT( "Editor" );
+
+		case EditorNoCommandlet:
+			return TEXT( "EditorNoCommandlet" );
+
+		case Program:
+			return TEXT("Program");
+
+		default:
+			ensureMsgf( false, TEXT( "Unrecognized EModuleType value: %i" ), Value );
+			return NULL;
+	}
+}
+
+FModuleDescriptor::FModuleDescriptor()
+{
+	Name = NAME_None;
+	Type = EHostType::Runtime;
+	LoadingPhase = ELoadingPhase::Default;
+}
+
+bool FModuleDescriptor::IsCompiledInCurrentConfiguration() const
+{
+	// Cache the string for the current platform
+	static FString UBTPlatform(FPlatformMisc::GetUBTPlatform());
+
+	// Check the platform is whitelisted
+	if(WhitelistPlatforms.Num() > 0 && !WhitelistPlatforms.Contains(UBTPlatform))
+	{
+		return false;
+	}
+
+	// Check the platform is not blacklisted
+	if(BlacklistPlatforms.Num() > 0 && BlacklistPlatforms.Contains(UBTPlatform))
+	{
+		return false;
+	}
+
+	// Check the module is compatible with this target. This should match UEBuildTarget.ShouldIncludePluginModule in UBT
+	switch (Type)
+	{
+	case EHostType::Runtime:
+	case EHostType::RuntimeNoCommandlet:
+		return true;
+
+	case EHostType::Developer:
+		#if WITH_UNREAL_DEVELOPER_TOOLS
+			return true;
+		#endif
+		break;
+
+	case EHostType::Editor:
+	case EHostType::EditorNoCommandlet:
+		#if WITH_EDITOR
+			return true;
+		#endif
+		break;
+
+	case EHostType::Program:
+		#if IS_PROGRAM
+			return true;
+		#endif
+		break;
+	}
+
+	return false;
+}
+
+bool FModuleDescriptor::IsLoadedInCurrentConfiguration() const
+{
+	// Check that the module is built for this configuration
+	if(!IsCompiledInCurrentConfiguration())
+	{
+		return false;
+	}
+
+	// Check that the runtime environment allows it to be loaded
+	switch (Type)
+	{
+	case EHostType::Runtime:
+		#if WITH_ENGINE || WITH_PLUGIN_SUPPORT
+			return true;
+		#endif
+		break;
+
+	case EHostType::RuntimeNoCommandlet:
+		#if WITH_ENGINE || WITH_PLUGIN_SUPPORT
+			if(!IsRunningCommandlet()) return true;
+		#endif
+		break;
+
+	case EHostType::Developer:
+		#if WITH_UNREAL_DEVELOPER_TOOLS
+			return true;
+		#endif
+		break;
+
+	case EHostType::Editor:
+		#if WITH_EDITOR
+			if(GIsEditor) return true;
+		#endif
+		break;
+
+	case EHostType::EditorNoCommandlet:
+		#if WITH_EDITOR
+			if(GIsEditor && !IsRunningCommandlet()) return true;
+		#endif
+		break;
+
+	case EHostType::Program:
+		#if WITH_PLUGIN_SUPPORT && IS_PROGRAM
+			return true;
+		#endif
+		break;
+	}
+
+	return false;
+}
+
+bool FModuleDescriptor::Read(const FJsonObject& Object, FText& OutFailReason)
+{
+	// Read the module name
+	TSharedPtr<FJsonValue> NameValue = Object.TryGetField(TEXT("Name"));
+	if(!NameValue.IsValid() || NameValue->Type != EJson::String)
+	{
+		OutFailReason = LOCTEXT("ModuleWithoutAName", "Found a 'Module' entry with a missing 'Name' field");
+		return false;
+	}
+	Name = FName(*NameValue->AsString());
+
+	// Read the module type
+	TSharedPtr<FJsonValue> TypeValue = Object.TryGetField(TEXT("Type"));
+	if(!TypeValue.IsValid() || TypeValue->Type != EJson::String)
+	{
+		OutFailReason = FText::Format( LOCTEXT( "ModuleWithoutAType", "Found Module entry '{0}' with a missing 'Type' field" ), FText::FromName(Name) );
+		return false;
+	}
+	Type = EHostType::FromString(*TypeValue->AsString());
+	if(Type == EHostType::Max)
+	{
+		OutFailReason = FText::Format( LOCTEXT( "ModuleWithInvalidType", "Module entry '{0}' specified an unrecognized module Type '{1}'" ), FText::FromName(Name), FText::FromString(TypeValue->AsString()) );
+		return false;
+	}
+
+	// Read the loading phase
+	TSharedPtr<FJsonValue> LoadingPhaseValue = Object.TryGetField(TEXT("LoadingPhase"));
+	if(LoadingPhaseValue.IsValid() && LoadingPhaseValue->Type == EJson::String)
+	{
+		LoadingPhase = ELoadingPhase::FromString(*LoadingPhaseValue->AsString());
+		if(LoadingPhase == ELoadingPhase::Max)
+		{
+			OutFailReason = FText::Format( LOCTEXT( "ModuleWithInvalidLoadingPhase", "Module entry '{0}' specified an unrecognized module LoadingPhase '{1}'" ), FText::FromName(Name), FText::FromString(LoadingPhaseValue->AsString()) );
+			return false;
+		}
+	}
+
+	// Read the whitelisted platforms
+	TSharedPtr<FJsonValue> WhitelistPlatformsValue = Object.TryGetField(TEXT("WhitelistPlatforms"));
+	if(WhitelistPlatformsValue.IsValid() && WhitelistPlatformsValue->Type == EJson::Array)
+	{
+		const TArray< TSharedPtr< FJsonValue > >& PlatformsArray = WhitelistPlatformsValue->AsArray();
+		for(int Idx = 0; Idx < PlatformsArray.Num(); Idx++)
+		{
+			WhitelistPlatforms.Add(PlatformsArray[Idx]->AsString());
+		}
+	}
+
+	// Read the blacklisted platforms
+	TSharedPtr<FJsonValue> BlacklistPlatformsValue = Object.TryGetField(TEXT("BlacklistPlatforms"));
+	if(BlacklistPlatformsValue.IsValid() && BlacklistPlatformsValue->Type == EJson::Array)
+	{
+		const TArray< TSharedPtr< FJsonValue > >& PlatformsArray = BlacklistPlatformsValue->AsArray();
+		for(int Idx = 0; Idx < PlatformsArray.Num(); Idx++)
+		{
+			BlacklistPlatforms.Add(PlatformsArray[Idx]->AsString());
+		}
+	}
+
+	return true;
+}
+
+bool FModuleDescriptor::ReadArray(const FJsonObject& Object, const TCHAR* Name, TArray<FModuleDescriptor>& OutModules, FText& OutFailReason)
+{
+	bool bResult = true;
+
+	TSharedPtr<FJsonValue> ModulesArrayValue = Object.TryGetField(Name);
+	if(ModulesArrayValue.IsValid() && ModulesArrayValue->Type == EJson::Array)
+	{
+		const TArray< TSharedPtr< FJsonValue > >& ModulesArray = ModulesArrayValue->AsArray();
+		for(int Idx = 0; Idx < ModulesArray.Num(); Idx++)
+		{
+			const TSharedPtr<FJsonValue>& ModuleValue = ModulesArray[Idx];
+			if(ModuleValue.IsValid() && ModuleValue->Type == EJson::Object)
+			{
+				FModuleDescriptor Descriptor;
+				if(Descriptor.Read(*ModuleValue->AsObject().Get(), OutFailReason))
+				{
+					OutModules.Add(Descriptor);
+				}
+				else
+				{
+					bResult = false;
+				}
+			}
+			else
+			{
+				OutFailReason = LOCTEXT( "ModuleWithInvalidModulesArray", "The 'Modules' array has invalid contents and was not able to be loaded." );
+				bResult = false;
+			}
+		}
+	}
+	
+	return bResult;
+}
+
+void FModuleDescriptor::Write(TJsonWriter<>& Writer) const
+{
+	Writer.WriteObjectStart();
+	Writer.WriteValue(TEXT("Name"), Name.ToString());
+	Writer.WriteValue(TEXT("Type"), FString(EHostType::ToString(Type)));
+	Writer.WriteValue(TEXT("LoadingPhase"), FString(ELoadingPhase::ToString(LoadingPhase)));
+	if (WhitelistPlatforms.Num() > 0)
+	{
+		Writer.WriteArrayStart(TEXT("WhitelistPlatforms"));
+		for(int Idx = 0; Idx < WhitelistPlatforms.Num(); Idx++)
+		{
+			Writer.WriteValue(WhitelistPlatforms[Idx]);
+		}
+		Writer.WriteArrayEnd();
+	}
+	if (BlacklistPlatforms.Num() > 0)
+	{
+		Writer.WriteArrayStart(TEXT("BlacklistPlatforms"));
+		for(int Idx = 0; Idx < BlacklistPlatforms.Num(); Idx++)
+		{
+			Writer.WriteValue(BlacklistPlatforms[Idx]);
+		}
+		Writer.WriteArrayEnd();
+	}
+	Writer.WriteObjectEnd();
+}
+
+void FModuleDescriptor::WriteArray(TJsonWriter<>& Writer, const TCHAR* Name, const TArray<FModuleDescriptor>& Modules)
+{
+	if(Modules.Num() > 0)
+	{
+		Writer.WriteArrayStart(Name);
+		for(int Idx = 0; Idx < Modules.Num(); Idx++)
+		{
+			Modules[Idx].Write(Writer);
+		}
+		Writer.WriteArrayEnd();
+	}
+}
+
+#undef LOCTEXT_NAMESPACE
+
diff --git a/Engine/Source/Runtime/Projects/Private/PluginManager.cpp b/Engine/Source/Runtime/Projects/Private/PluginManager.cpp
--- a/Engine/Source/Runtime/Projects/Private/PluginManager.cpp
+++ b/Engine/Source/Runtime/Projects/Private/PluginManager.cpp
@@ -246,7 +246,7 @@ void FPluginManager::DiscoverAllPlugins()
 			// Make sure to tell the module manager about any of our plugin's modules, so it will know where on disk to find them.
 			for( auto ModuleInfoIt( PluginInfo.Modules.CreateConstIterator() ); ModuleInfoIt; ++ModuleInfoIt )
 			{
-				const FProjectOrPluginInfo::FModuleInfo& ModuleInfo = *ModuleInfoIt;
+				const FModuleDescriptor& ModuleInfo = *ModuleInfoIt;
 
 				// @todo plugin: We're adding all modules always here, even editor modules which might not exist at runtime. Hope that's OK.
 				FModuleManager::Get().AddModule(ModuleInfo.Name);
diff --git a/Engine/Source/Runtime/Projects/Private/PluginManagerShared.cpp b/Engine/Source/Runtime/Projects/Private/PluginManagerShared.cpp
--- a/Engine/Source/Runtime/Projects/Private/PluginManagerShared.cpp
+++ b/Engine/Source/Runtime/Projects/Private/PluginManagerShared.cpp
@@ -11,14 +11,14 @@ void FProjectOrPlugin::LoadModules( const ELoadingPhase::Type LoadingPhase, TMap
 {
 	const FProjectOrPluginInfo& ProjectOrPluginInfo = GetProjectOrPluginInfo();
 
-	for( TArray< FProjectOrPluginInfo::FModuleInfo >::TConstIterator ModuleInfoIt( ProjectOrPluginInfo.Modules.CreateConstIterator() ); ModuleInfoIt; ++ModuleInfoIt )
+	for( TArray< FModuleDescriptor >::TConstIterator ModuleInfoIt( ProjectOrPluginInfo.Modules.CreateConstIterator() ); ModuleInfoIt; ++ModuleInfoIt )
 	{
-		const FProjectOrPluginInfo::FModuleInfo& ModuleInfo = *ModuleInfoIt;
+		const FModuleDescriptor& ModuleInfo = *ModuleInfoIt;
 
 		// Don't need to do anything if this module is already loaded
 		if( !FModuleManager::Get().IsModuleLoaded( ModuleInfo.Name ) )
 		{
-			if( LoadingPhase == ModuleInfo.LoadingPhase && ShouldLoadModule(ModuleInfo) )
+			if( LoadingPhase == ModuleInfo.LoadingPhase && ModuleInfo.IsLoadedInCurrentConfiguration() )
 			{
 				// @todo plugin: DLL search problems.  Plugins that statically depend on other modules within this plugin may not be found?  Need to test this.
 
@@ -186,38 +186,7 @@ FString FProjectOrPlugin::SerializeToJSON( ) const
 	Writer->WriteValue(TEXT("CreatedBy"), ProjectOrPluginInfo.CreatedBy);
 	Writer->WriteValue(TEXT("CreatedByURL"), ProjectOrPluginInfo.CreatedByURL);
 	
-	if ( ProjectOrPluginInfo.Modules.Num() > 0 )
-	{
-		Writer->WriteArrayStart(TEXT("Modules"));
-		for ( auto ModuleIt = ProjectOrPluginInfo.Modules.CreateConstIterator(); ModuleIt; ++ModuleIt )
-		{
-			const FProjectOrPluginInfo::FModuleInfo& Module = *ModuleIt;
-			Writer->WriteObjectStart();
-			Writer->WriteValue(TEXT("Name"), Module.Name.ToString());
-			Writer->WriteValue(TEXT("Type"), FString(FProjectOrPluginInfo::EModuleType::ToString(Module.Type)));
-			Writer->WriteValue(TEXT("LoadingPhase"), FString(ELoadingPhase::ToString(Module.LoadingPhase)));
-			if (Module.WhitelistPlatforms.Num() > 0)
-			{
-				Writer->WriteArrayStart(TEXT("WhitelistPlatforms"));
-				for ( auto PlatformIt = Module.WhitelistPlatforms.CreateConstIterator(); PlatformIt; ++PlatformIt )
-				{
-					Writer->WriteValue(*PlatformIt);
-				}
-				Writer->WriteArrayEnd();
-			}
-			if (Module.BlacklistPlatforms.Num() > 0)
-			{
-				Writer->WriteArrayStart(TEXT("BlacklistPlatforms"));
-				for ( auto PlatformIt = Module.BlacklistPlatforms.CreateConstIterator(); PlatformIt; ++PlatformIt )
-				{
-					Writer->WriteValue(*PlatformIt);
-				}
-				Writer->WriteArrayEnd();
-			}
-			Writer->WriteObjectEnd();
-		}
-		Writer->WriteArrayEnd();
-	}
+	FModuleDescriptor::WriteArray(Writer.Get(), TEXT("Modules"), ProjectOrPluginInfo.Modules);
 	
 	PerformAdditionalSerialization(Writer);
 	
@@ -236,10 +205,10 @@ bool FProjectOrPlugin::RequiresUpdate( ) const
 bool FProjectOrPlugin::AreModulesUpToDate( ) const
 {
 	const FProjectOrPluginInfo& ProjectOrPluginInfo = GetProjectOrPluginInfo();
-	for (TArray< FProjectOrPluginInfo::FModuleInfo >::TConstIterator Iter(ProjectOrPluginInfo.Modules); Iter; ++Iter)
+	for (TArray< FModuleDescriptor >::TConstIterator Iter(ProjectOrPluginInfo.Modules); Iter; ++Iter)
 	{
-		const FProjectOrPluginInfo::FModuleInfo &ModuleInfo = *Iter;
-		if (ShouldBuildModule(ModuleInfo) && !FModuleManager::Get().IsModuleUpToDate(ModuleInfo.Name))
+		const FModuleDescriptor &ModuleInfo = *Iter;
+		if (ModuleInfo.IsCompiledInCurrentConfiguration() && !FModuleManager::Get().IsModuleUpToDate(ModuleInfo.Name))
 		{
 			return false;
 		}
@@ -267,111 +236,13 @@ void FProjectOrPlugin::ReplaceModulesInProject(const TArray<FString>* StartupMod
 		ProjectOrPluginInfo.Modules.Empty();
 		for ( auto ModuleIt = StartupModuleNames->CreateConstIterator(); ModuleIt; ++ModuleIt )
 		{
-			FProjectOrPluginInfo::FModuleInfo ModuleInfo;
+			FModuleDescriptor ModuleInfo;
 			ModuleInfo.Name = FName(**ModuleIt);
 			ProjectOrPluginInfo.Modules.Add(ModuleInfo);
 		}
 	}
 }
 
-bool FProjectOrPlugin::ShouldBuildModule(const FProjectOrPluginInfo::FModuleInfo& ModuleInfo) const
-{
-	// Cache the string for the current platform
-	static FString UBTPlatform(FPlatformMisc::GetUBTPlatform());
-
-	// Check the platform is whitelisted
-	if(ModuleInfo.WhitelistPlatforms.Num() > 0 && !ModuleInfo.WhitelistPlatforms.Contains(UBTPlatform))
-	{
-		return false;
-	}
-
-	// Check the platform is not blacklisted
-	if(ModuleInfo.BlacklistPlatforms.Num() > 0 && ModuleInfo.BlacklistPlatforms.Contains(UBTPlatform))
-	{
-		return false;
-	}
-
-	// Check the module is compatible with this target. This should match UEBuildTarget.ShouldIncludePluginModule in UBT
-	switch (ModuleInfo.Type)
-	{
-	case FProjectOrPluginInfo::EModuleType::Runtime:
-	case FProjectOrPluginInfo::EModuleType::RuntimeNoCommandlet:
-		return true;
-
-	case FProjectOrPluginInfo::EModuleType::Developer:
-		#if WITH_UNREAL_DEVELOPER_TOOLS
-			return true;
-		#endif
-		break;
-
-	case FProjectOrPluginInfo::EModuleType::Editor:
-	case FProjectOrPluginInfo::EModuleType::EditorNoCommandlet:
-		#if WITH_EDITOR
-			return true;
-		#endif
-		break;
-
-	case FProjectOrPluginInfo::EModuleType::Program:
-		#if IS_PROGRAM
-			return true;
-		#endif
-		break;
-	}
-
-	return false;
-}
-
-bool FProjectOrPlugin::ShouldLoadModule(const FProjectOrPluginInfo::FModuleInfo& ModuleInfo) const
-{
-	// Check that the module is built for this configuration
-	if(!ShouldBuildModule(ModuleInfo))
-	{
-		return false;
-	}
-
-	// Check that the runtime environment allows it to be loaded
-	switch (ModuleInfo.Type)
-	{
-	case FProjectOrPluginInfo::EModuleType::Runtime:
-		#if WITH_ENGINE || WITH_PLUGIN_SUPPORT
-			return true;
-		#endif
-		break;
-
-	case FProjectOrPluginInfo::EModuleType::RuntimeNoCommandlet:
-		#if WITH_ENGINE || WITH_PLUGIN_SUPPORT
-			if(!IsRunningCommandlet()) return true;
-		#endif
-		break;
-
-	case FProjectOrPluginInfo::EModuleType::Developer:
-		#if WITH_UNREAL_DEVELOPER_TOOLS
-			return true;
-		#endif
-		break;
-
-	case FProjectOrPluginInfo::EModuleType::Editor:
-		#if WITH_EDITOR
-			if(GIsEditor) return true;
-		#endif
-		break;
-
-	case FProjectOrPluginInfo::EModuleType::EditorNoCommandlet:
-		#if WITH_EDITOR
-			if(GIsEditor && !IsRunningCommandlet()) return true;
-		#endif
-		break;
-
-	case FProjectOrPluginInfo::EModuleType::Program:
-		#if WITH_PLUGIN_SUPPORT && IS_PROGRAM
-			return true;
-		#endif
-		break;
-	}
-
-	return false;
-}
-
 bool FProjectOrPlugin::ReadNumberFromJSON(const TSharedRef< FJsonObject >& FileObject, const FString& PropertyName, int32& OutNumber ) const
 {
 	int64 LargeNumber;
@@ -472,158 +343,20 @@ bool FProjectOrPlugin::ReadFileVersionFromJSON(const TSharedRef< FJsonObject >&
 	return false;
 }
 
-bool FProjectOrPlugin::ReadModulesFromJSON(const TSharedRef< FJsonObject >& FileObject, TArray<FProjectOrPluginInfo::FModuleInfo>& OutModules, FText& OutFailReason ) const
+bool FProjectOrPlugin::ReadModulesFromJSON(const TSharedRef< FJsonObject >& FileObject, TArray<FModuleDescriptor>& OutModules, FText& OutFailReason ) const
 {
-	bool bLoadedSuccessfully = true;
-
-	// This project or plugin might have some modules that we need to know about.  Let's take a look.
-	if( FileObject->HasField( TEXT( "Modules" ) ) )
-	{
-		const TSharedPtr< FJsonValue >& ModulesValue = FileObject->GetField< EJson::Array >( TEXT( "Modules" ) );
-		const TArray< TSharedPtr< FJsonValue > >& ModulesArray = ModulesValue->AsArray();
-		for( auto ModuleValueIt( ModulesArray.CreateConstIterator() ); ModuleValueIt; ++ModuleValueIt )
-		{
-			const TSharedPtr< FJsonObject >& ModuleObjectPtr = ( *ModuleValueIt )->AsObject();
-			if( ModuleObjectPtr.IsValid() )
-			{
-				const TSharedRef< FJsonObject >& ModuleObject = ModuleObjectPtr.ToSharedRef();
-				FProjectOrPluginInfo::FModuleInfo ModuleInfo;
-
-				// Module name
-				{
-					// All modules require a name to be set
-					FString ModuleName;
-					if ( ReadStringFromJSON(ModuleObject, TEXT("Name"), ModuleName) )
-					{
-						ModuleInfo.Name = FName(*ModuleName);
-					}
-					else
-					{
-						OutFailReason = LOCTEXT("ModuleWithoutAName", "Found a 'Module' entry with a missing 'Name' field");
-						bLoadedSuccessfully = false;
-						continue;
-					}
-				}
-
-
-				// Module type
-				{
-					FString ModuleType;
-					if ( ReadStringFromJSON(ModuleObject, TEXT("Type"), ModuleType) )
-					{
-						// Check to see if this is a valid type
-						bool bFoundValidType = false;
-						for( int32 PossibleTypeIndex = 0; PossibleTypeIndex < FProjectOrPluginInfo::EModuleType::Max; ++PossibleTypeIndex )
-						{
-							const FProjectOrPluginInfo::EModuleType::Type PossibleType = (FProjectOrPluginInfo::EModuleType::Type)PossibleTypeIndex;
-							const TCHAR* PossibleTypeName = FProjectOrPluginInfo::EModuleType::ToString( PossibleType );
-
-							if( FCString::Stricmp( PossibleTypeName, *ModuleType ) == 0 )
-							{
-								bFoundValidType = true;
-								ModuleInfo.Type = PossibleType;
-								break;
-							}
-						}
-
-						if( !bFoundValidType )
-						{
-							OutFailReason = FText::Format( LOCTEXT( "ModuleWithInvalidType", "Module entry '{0}' specified an unrecognized module Type '{1}'" ), FText::FromName(ModuleInfo.Name), FText::FromString(ModuleType) );
-							bLoadedSuccessfully = false;
-							continue;
-						}
-					}
-					else
-					{
-						OutFailReason = FText::Format( LOCTEXT( "ModuleWithoutAType", "Found Module entry '{0}' with a missing 'Type' field" ), FText::FromName(ModuleInfo.Name) );
-						bLoadedSuccessfully = false;
-						break;
-					}
-				}
-
-				// Loading phase
-				{
-					FString ModuleLoadingPhase;
-					if ( ReadStringFromJSON(ModuleObject, TEXT("LoadingPhase"), ModuleLoadingPhase) )
-					{
-						// Check to see if this is a valid Phase
-						bool bFoundValidPhase = false;
-						for( int32 PossiblePhaseIndex = 0; PossiblePhaseIndex < ELoadingPhase::Max; ++PossiblePhaseIndex )
-						{
-							const ELoadingPhase::Type PossiblePhase = (ELoadingPhase::Type)PossiblePhaseIndex;
-							const TCHAR* PossiblePhaseName = ELoadingPhase::ToString( PossiblePhase );
-
-							if( FCString::Stricmp( PossiblePhaseName, *ModuleLoadingPhase ) == 0 )
-							{
-								bFoundValidPhase = true;
-								ModuleInfo.LoadingPhase = PossiblePhase;
-								break;
-							}
-						}
-
-						if( !bFoundValidPhase )
-						{
-							OutFailReason = FText::Format( LOCTEXT( "ModuleWithInvalidLoadingPhase", "Module entry '{0}' specified an unrecognized module LoadingPhase '{1}'" ), FText::FromName(ModuleInfo.Name), FText::FromString(ModuleLoadingPhase) );
-							bLoadedSuccessfully = false;
-							continue;
-						}
-					}
-					else
-					{
-						// No 'LoadingPhase' was specified in the file.  That's totally fine, we'll use the default.
-					}
-				}
-
-				// Whitelist platforms
-				{
-					if( ModuleObject->HasField( TEXT( "WhitelistPlatforms" ) ) )
-					{
-						// walk over the array values
-						const TSharedPtr< FJsonValue >& PlatformsValue = ModuleObject->GetField< EJson::Array >( TEXT( "WhitelistPlatforms" ) );
-						const TArray< TSharedPtr< FJsonValue > >& PlatformsArray = PlatformsValue->AsArray();
-						for( auto PlatformValueIt( PlatformsArray.CreateConstIterator() ); PlatformValueIt; ++PlatformValueIt )
-						{
-							ModuleInfo.WhitelistPlatforms.Add((*PlatformValueIt)->AsString());
-						}
-					}
-				}
-
-				// Blacklist platforms
-				{
-					if( ModuleObject->HasField( TEXT( "BlacklistPlatforms" ) ) )
-					{
-						// walk over the array values
-						const TSharedPtr< FJsonValue >& PlatformsValue = ModuleObject->GetField< EJson::Array >( TEXT( "BlacklistPlatforms" ) );
-						const TArray< TSharedPtr< FJsonValue > >& PlatformsArray = PlatformsValue->AsArray();
-						for( auto PlatformValueIt( PlatformsArray.CreateConstIterator() ); PlatformValueIt; ++PlatformValueIt )
-						{
-							ModuleInfo.BlacklistPlatforms.Add((*PlatformValueIt)->AsString());
-						}
-					}
-				}
-
-				OutModules.Add(ModuleInfo);
-			}
-			else
-			{
-				OutFailReason = LOCTEXT( "ModuleWithInvalidModulesArray", "The 'Modules' array has invalid contents and was not able to be loaded." );
-				bLoadedSuccessfully = false;
-			}
-		}
-	}
-	
-	return bLoadedSuccessfully;
+	return FModuleDescriptor::ReadArray(FileObject.Get(), TEXT("Modules"), OutModules, OutFailReason);
 }
 
 
 
-ELoadingPhase::Type FProjectAndPluginManager::GetEarliestPhaseFromModules(const TArray<FProjectOrPluginInfo::FModuleInfo>& Modules)
+ELoadingPhase::Type FProjectAndPluginManager::GetEarliestPhaseFromModules(const TArray<FModuleDescriptor>& Modules)
 {
 	ELoadingPhase::Type EarliestPhase = ELoadingPhase::Default;
 
 	for ( auto ModuleIt = Modules.CreateConstIterator(); ModuleIt; ++ModuleIt )
 	{
-		EarliestPhase = ELoadingPhase::GetEarlierPhase(EarliestPhase, (*ModuleIt).LoadingPhase);
+		EarliestPhase = ELoadingPhase::Earliest(EarliestPhase, (*ModuleIt).LoadingPhase);
 	}
 
 	return EarliestPhase;
diff --git a/Engine/Source/Runtime/Projects/Private/PluginManagerShared.h b/Engine/Source/Runtime/Projects/Private/PluginManagerShared.h
--- a/Engine/Source/Runtime/Projects/Private/PluginManagerShared.h
+++ b/Engine/Source/Runtime/Projects/Private/PluginManagerShared.h
@@ -24,85 +24,8 @@ enum EProjectFileVersion
 class FProjectOrPluginInfo
 {
 public:
-	struct EModuleType
-	{
-		enum Type
-		{
-			Runtime,
-			RuntimeNoCommandlet,
-			Developer,
-			Editor,
-			EditorNoCommandlet,
-			/** Program-only plugin type */
-			Program,
-			// NOTE: If you add a new value, make sure to update the ToString() method below!
-
-			Max
-		};
-
-
-		/**
-		 * Converts an EModuleType::Type value to a string literal
-		 *
-		 * @param	The value to convert to a string
-		 *
-		 * @return	The string representation of this enum value
-		 */
-		inline static const TCHAR* ToString( const EModuleType::Type Value )
-		{
-			switch( Value )
-			{
-				case Runtime:
-					return TEXT( "Runtime" );
-
-				case RuntimeNoCommandlet:
-					return TEXT( "RuntimeNoCommandlet" );
-
-				case Developer:
-					return TEXT( "Developer" );
-
-				case Editor:
-					return TEXT( "Editor" );
-
-				case EditorNoCommandlet:
-					return TEXT( "EditorNoCommandlet" );
-
-				case Program:
-					return TEXT("Program");
-
-				default:
-					ensureMsgf( false, TEXT( "Unrecognized EModuleType value: %i" ), Value );
-					return NULL;
-			}
-		}
-	};
-
-
-	struct FModuleInfo
-	{
-		/** Name of this module */
-		FName Name;
-
-		/** Type of module */
-		EModuleType::Type Type;
-
-		/** When should the module be loaded during the startup sequence?  This is sort of an advanced setting. */
-		ELoadingPhase::Type LoadingPhase;
 
-		/** List of allowed platforms */
-		TArray<FString> WhitelistPlatforms;
 
-		/** List of disallowed platforms */
-		TArray<FString> BlacklistPlatforms;
-
-		/** Module info constructor */
-		FModuleInfo()
-			: Name( NAME_None ),
-			  Type( EModuleType::Runtime ),
-			  LoadingPhase( ELoadingPhase::Default )
-		{
-		}
-	};
 
 	/** Name of the project or plugin, which is actually just the base name of the .uproject or .uplugin file.  Note that this is not guaranteed to
 	    be globally unique!  Projects or plugins with the same name could exist in different directories. */
@@ -162,7 +85,7 @@ public:
 	FString CreatedByURL;
 
 	/** List of all modules associated with this project or plugin */
-	TArray<FModuleInfo> Modules;
+	TArray<FModuleDescriptor> Modules;
 
 
 	/** ProjectOrPlugin info constructor */
@@ -248,12 +171,6 @@ public:
 	void ReplaceModulesInProject(const TArray<FString>* StartupModuleNames);
 
 protected:
-	/** Checks whether a given module should be built for this build configuration */
-	bool ShouldBuildModule(const FProjectOrPluginInfo::FModuleInfo& ModuleInfo) const;
-
-	/** Checks whether a given module should be loaded for this build configuration */
-	bool ShouldLoadModule(const FProjectOrPluginInfo::FModuleInfo& ModuleInfo) const;
-
 	/** @return Exposes access to the project or plugin's descriptor */
 	virtual FProjectOrPluginInfo& GetProjectOrPluginInfo() = 0;
 	virtual const FProjectOrPluginInfo& GetProjectOrPluginInfo() const = 0;
@@ -271,7 +188,7 @@ protected:
 private:
 	/** Helper functions to read specific special values from the JsonObject */
 	bool ReadFileVersionFromJSON(const TSharedRef< FJsonObject >& FileObject, int32& OutVersion ) const;
-	bool ReadModulesFromJSON(const TSharedRef< FJsonObject >& FileObject, TArray<FProjectOrPluginInfo::FModuleInfo>& OutModules, FText& OutFailReason ) const;
+	bool ReadModulesFromJSON(const TSharedRef< FJsonObject >& FileObject, TArray<FModuleDescriptor>& OutModules, FText& OutFailReason ) const;
 };
 
 /**
@@ -281,7 +198,7 @@ class FProjectAndPluginManager
 {
 public:
 	/** Returns the earliest phase in the specified list of modules */
-	static ELoadingPhase::Type GetEarliestPhaseFromModules(const TArray<FProjectOrPluginInfo::FModuleInfo>& Modules);
+	static ELoadingPhase::Type GetEarliestPhaseFromModules(const TArray<FModuleDescriptor>& Modules);
 };
 
 
diff --git a/Engine/Source/Runtime/Projects/Private/ProjectManager.cpp b/Engine/Source/Runtime/Projects/Private/ProjectManager.cpp
--- a/Engine/Source/Runtime/Projects/Private/ProjectManager.cpp
+++ b/Engine/Source/Runtime/Projects/Private/ProjectManager.cpp
@@ -223,7 +223,7 @@ bool FProjectManager::DuplicateProjectFile(const FString& SourceProjectFilename,
 	const FString BaseNewName = FPaths::GetBaseFilename(NewProjectFilename);
 	for ( auto ModuleIt = ProjectInfo.Modules.CreateIterator(); ModuleIt; ++ModuleIt )
 	{
-		FProjectOrPluginInfo::FModuleInfo& ModuleInfo = *ModuleIt;
+		FModuleDescriptor& ModuleInfo = *ModuleIt;
 		ModuleInfo.Name = FName(*ModuleInfo.Name.ToString().Replace(*BaseSourceName, *BaseNewName));
 	}
 
diff --git a/Engine/Source/Runtime/Projects/Public/Interfaces/IPluginManager.h b/Engine/Source/Runtime/Projects/Public/Interfaces/IPluginManager.h
--- a/Engine/Source/Runtime/Projects/Public/Interfaces/IPluginManager.h
+++ b/Engine/Source/Runtime/Projects/Public/Interfaces/IPluginManager.h
@@ -3,7 +3,7 @@
 #pragma once
 
 
-#include "IPluginManagerShared.h"
+#include "ModuleDescriptor.h"
 
 
 /**
diff --git a/Engine/Source/Runtime/Projects/Public/Interfaces/IPluginManagerShared.h b/Engine/Source/Runtime/Projects/Public/Interfaces/IPluginManagerShared.h
deleted file mode 100644
diff --git a/Engine/Source/Runtime/Projects/Public/Interfaces/IProjectManager.h b/Engine/Source/Runtime/Projects/Public/Interfaces/IProjectManager.h
--- a/Engine/Source/Runtime/Projects/Public/Interfaces/IProjectManager.h
+++ b/Engine/Source/Runtime/Projects/Public/Interfaces/IProjectManager.h
@@ -3,8 +3,8 @@
 #pragma once
 
 
-#include "IPluginManagerShared.h"
 #include "EngineVersion.h"
+#include "ModuleDescriptor.h"
 
 /**
  * Simple data structure that is filled when querying information about projects
diff --git a/Engine/Source/Runtime/Projects/Public/ModuleDescriptor.h b/Engine/Source/Runtime/Projects/Public/ModuleDescriptor.h
new file mode 100644
--- /dev/null
+++ b/Engine/Source/Runtime/Projects/Public/ModuleDescriptor.h
@@ -0,0 +1,134 @@
+// Copyright 1998-2014 Epic Games, Inc. All Rights Reserved.
+
+#pragma once
+
+#include "Json.h"
+
+/**
+ * Phase at which this module should be loaded during startup.
+ */
+namespace ELoadingPhase
+{
+	enum Type
+	{
+		/** Loaded at the default loading point during startup (during engine init, after game modules are loaded.) */
+		Default,
+
+		/** Right after the default phase */
+		PostDefault,
+
+		/** Right before the default phase */
+		PreDefault,
+
+		/** Loaded before the engine is fully initialized, immediately after the config system has been initialized.  Necessary only for very low-level hooks */
+		PostConfigInit,
+
+		/** Loaded before the engine is fully initialized for modules that need to hook into the loading screen before it triggers */
+		PreLoadingScreen,
+		// NOTE: If you add a new value, make sure to update the ToString() method below!
+
+		Max
+	};
+
+	/**
+	 * Returns the earliest of two load phases.
+	 *
+	 * @param	The value to convert to a string
+	 * @return	The string representation of this enum value
+	 */
+	ELoadingPhase::Type Earliest( ELoadingPhase::Type First, ELoadingPhase::Type Second );
+	
+	/**
+	 * Converts a string to a ELoadingPhase::Type value
+	 *
+	 * @param	The string to convert to a value
+	 * @return	The corresponding value, or 'Max' if the string is not valid.
+	 */
+	ELoadingPhase::Type FromString( const TCHAR *Text );
+
+	/**
+	 * Returns the name of a module load phase.
+	 *
+	 * @param	The value to convert to a string
+	 * @return	The string representation of this enum value
+	 */
+	const TCHAR* ToString( const ELoadingPhase::Type Value );
+};
+
+/**
+ * Environment that can load a module.
+ */
+namespace EHostType
+{
+	enum Type
+	{
+		Runtime,
+		RuntimeNoCommandlet,
+		Developer,
+		Editor,
+		EditorNoCommandlet,
+		Program,		//!< Program-only plugin type
+
+		// NOTE: If you add a new value, make sure to update the ToString() method below!
+
+		Max
+	};
+
+	/**
+	 * Converts a string to a EHostType::Type value
+	 *
+	 * @param	The string to convert to a value
+	 * @return	The corresponding value, or 'Max' if the string is not valid.
+	 */
+	EHostType::Type FromString( const TCHAR *Text );
+
+	/**
+	 * Converts an EHostType::Type value to a string literal
+	 *
+	 * @param	The value to convert to a string
+	 * @return	The string representation of this enum value
+	 */
+	const TCHAR* ToString( const EHostType::Type Value );
+};
+
+/**
+ * Description of a loadable module.
+ */
+struct FModuleDescriptor
+{
+	/** Name of this module */
+	FName Name;
+
+	/** Usage type of module */
+	EHostType::Type Type;
+
+	/** When should the module be loaded during the startup sequence?  This is sort of an advanced setting. */
+	ELoadingPhase::Type LoadingPhase;
+
+	/** List of allowed platforms */
+	TArray<FString> WhitelistPlatforms;
+
+	/** List of disallowed platforms */
+	TArray<FString> BlacklistPlatforms;
+
+	/** Module info constructor */
+	FModuleDescriptor();
+
+	/** Reads a descriptor from the given JSON object */
+	bool Read(const FJsonObject& Object, FText& OutFailReason);
+
+	/** Reads an array of modules from the given JSON object */
+	static bool ReadArray(const FJsonObject& Object, const TCHAR* Name, TArray<FModuleDescriptor>& OutModules, FText& OutFailReason);
+
+	/** Writes a descriptor to JSON */
+	void Write(TJsonWriter<>& Writer) const;
+
+	/** Writes an array of modules to JSON */
+	static void WriteArray(TJsonWriter<>& Writer, const TCHAR* Name, const TArray<FModuleDescriptor>& Modules);
+
+	/** Tests whether the module should be built for the current engine configuration */
+	bool IsCompiledInCurrentConfiguration() const;
+
+	/** Tests whether the module should be loaded for the current engine configuration */
+	bool IsLoadedInCurrentConfiguration() const;
+};
diff --git a/Engine/Source/Runtime/Projects/Public/Projects.h b/Engine/Source/Runtime/Projects/Public/Projects.h
--- a/Engine/Source/Runtime/Projects/Public/Projects.h
+++ b/Engine/Source/Runtime/Projects/Public/Projects.h
@@ -12,6 +12,6 @@
 /* Public Includes
  *****************************************************************************/
 
-#include "IPluginManagerShared.h"
+#include "ModuleDescriptor.h"
 #include "IPluginManager.h"
 #include "IProjectManager.h"
