Separate out project and plugin descriptors from each other, and the instance data maintained by their respective managers. Projects and plugins already only have superficial similarities, and will continue to diverge.

diff --git a/Engine/Source/Runtime/Projects/Private/ModuleDescriptor.cpp b/Engine/Source/Runtime/Projects/Private/ModuleDescriptor.cpp
--- a/Engine/Source/Runtime/Projects/Private/ModuleDescriptor.cpp
+++ b/Engine/Source/Runtime/Projects/Private/ModuleDescriptor.cpp
@@ -2,40 +2,7 @@
 
 #include "ProjectsPrivatePCH.h"
 
-#define LOCTEXT_NAMESPACE "PluginManagerShared"
+#define LOCTEXT_NAMESPACE "ModuleDescriptor"
 
-static int32 GetLoadOrderValue( ELoadingPhase::Type Phase )
-{
-	switch( Phase )
-	{
-	case ELoadingPhase::PostDefault:
-		return 500;
-
-	case ELoadingPhase::Default:
-		return 400;
-
-	case ELoadingPhase::PreDefault:
-		return 300;
-			
-	case ELoadingPhase::PreLoadingScreen:
-		return 200;
-
-	case ELoadingPhase::PostConfigInit:
-		return 100;
-
-	default:
-		ensureMsgf( false, TEXT( "Unrecognized ELoadingPhase value: %i" ), Phase );
-		return 0;
-	}
-}
-
-ELoadingPhase::Type ELoadingPhase::Earliest(ELoadingPhase::Type PhaseA, ELoadingPhase::Type PhaseB)
-{
-	const int32 PhaseAValue = GetLoadOrderValue(PhaseA);
-	const int32 PhaseBValue = GetLoadOrderValue(PhaseB);
-
-	return PhaseAValue < PhaseBValue ? PhaseA : PhaseB;
-}
-
 ELoadingPhase::Type ELoadingPhase::FromString( const TCHAR *String )
 {
@@ -118,109 +85,11 @@
 	}
 }
 
-FModuleDescriptor::FModuleDescriptor()
+FModuleDescriptor::FModuleDescriptor(const FName InName, EHostType::Type InType, ELoadingPhase::Type InLoadingPhase)
+	: Name(InName)
+	, Type(InType)
+	, LoadingPhase(InLoadingPhase)
 {
-	Name = NAME_None;
-	Type = EHostType::Runtime;
-	LoadingPhase = ELoadingPhase::Default;
 }
 
-bool FModuleDescriptor::IsCompiledInCurrentConfiguration() const
-{
-	// Cache the string for the current platform
-	static FString UBTPlatform(FPlatformMisc::GetUBTPlatform());
-
-	// Check the platform is whitelisted
-	if(WhitelistPlatforms.Num() > 0 && !WhitelistPlatforms.Contains(UBTPlatform))
-	{
-		return false;
-	}
-
-	// Check the platform is not blacklisted
-	if(BlacklistPlatforms.Num() > 0 && BlacklistPlatforms.Contains(UBTPlatform))
-	{
-		return false;
-	}
-
-	// Check the module is compatible with this target. This should match UEBuildTarget.ShouldIncludePluginModule in UBT
-	switch (Type)
-	{
-	case EHostType::Runtime:
-	case EHostType::RuntimeNoCommandlet:
-		return true;
-
-	case EHostType::Developer:
-		#if WITH_UNREAL_DEVELOPER_TOOLS
-			return true;
-		#endif
-		break;
-
-	case EHostType::Editor:
-	case EHostType::EditorNoCommandlet:
-		#if WITH_EDITOR
-			return true;
-		#endif
-		break;
-
-	case EHostType::Program:
-		#if IS_PROGRAM
-			return true;
-		#endif
-		break;
-	}
-
-	return false;
-}
-
-bool FModuleDescriptor::IsLoadedInCurrentConfiguration() const
-{
-	// Check that the module is built for this configuration
-	if(!IsCompiledInCurrentConfiguration())
-	{
-		return false;
-	}
-
-	// Check that the runtime environment allows it to be loaded
-	switch (Type)
-	{
-	case EHostType::Runtime:
-		#if WITH_ENGINE || WITH_PLUGIN_SUPPORT
-			return true;
-		#endif
-		break;
-
-	case EHostType::RuntimeNoCommandlet:
-		#if WITH_ENGINE || WITH_PLUGIN_SUPPORT
-			if(!IsRunningCommandlet()) return true;
-		#endif
-		break;
-
-	case EHostType::Developer:
-		#if WITH_UNREAL_DEVELOPER_TOOLS
-			return true;
-		#endif
-		break;
-
-	case EHostType::Editor:
-		#if WITH_EDITOR
-			if(GIsEditor) return true;
-		#endif
-		break;
-
-	case EHostType::EditorNoCommandlet:
-		#if WITH_EDITOR
-			if(GIsEditor && !IsRunningCommandlet()) return true;
-		#endif
-		break;
-
-	case EHostType::Program:
-		#if WITH_PLUGIN_SUPPORT && IS_PROGRAM
-			return true;
-		#endif
-		break;
-	}
-
-	return false;
-}
-
 bool FModuleDescriptor::Read(const FJsonObject& Object, FText& OutFailReason)
@@ -359,5 +228,147 @@
 	}
 }
 
+bool FModuleDescriptor::IsCompiledInCurrentConfiguration() const
+{
+	// Cache the string for the current platform
+	static FString UBTPlatform(FPlatformMisc::GetUBTPlatform());
+
+	// Check the platform is whitelisted
+	if(WhitelistPlatforms.Num() > 0 && !WhitelistPlatforms.Contains(UBTPlatform))
+	{
+		return false;
+	}
+
+	// Check the platform is not blacklisted
+	if(BlacklistPlatforms.Num() > 0 && BlacklistPlatforms.Contains(UBTPlatform))
+	{
+		return false;
+	}
+
+	// Check the module is compatible with this target. This should match UEBuildTarget.ShouldIncludePluginModule in UBT
+	switch (Type)
+	{
+	case EHostType::Runtime:
+	case EHostType::RuntimeNoCommandlet:
+		return true;
+
+	case EHostType::Developer:
+		#if WITH_UNREAL_DEVELOPER_TOOLS
+			return true;
+		#endif
+		break;
+
+	case EHostType::Editor:
+	case EHostType::EditorNoCommandlet:
+		#if WITH_EDITOR
+			return true;
+		#endif
+		break;
+
+	case EHostType::Program:
+		#if IS_PROGRAM
+			return true;
+		#endif
+		break;
+	}
+
+	return false;
+}
+
+bool FModuleDescriptor::IsLoadedInCurrentConfiguration() const
+{
+	// Check that the module is built for this configuration
+	if(!IsCompiledInCurrentConfiguration())
+	{
+		return false;
+	}
+
+	// Check that the runtime environment allows it to be loaded
+	switch (Type)
+	{
+	case EHostType::Runtime:
+		#if WITH_ENGINE || WITH_PLUGIN_SUPPORT
+			return true;
+		#endif
+		break;
+
+	case EHostType::RuntimeNoCommandlet:
+		#if WITH_ENGINE || WITH_PLUGIN_SUPPORT
+			if(!IsRunningCommandlet()) return true;
+		#endif
+		break;
+
+	case EHostType::Developer:
+		#if WITH_UNREAL_DEVELOPER_TOOLS
+			return true;
+		#endif
+		break;
+
+	case EHostType::Editor:
+		#if WITH_EDITOR
+			if(GIsEditor) return true;
+		#endif
+		break;
+
+	case EHostType::EditorNoCommandlet:
+		#if WITH_EDITOR
+			if(GIsEditor && !IsRunningCommandlet()) return true;
+		#endif
+		break;
+
+	case EHostType::Program:
+		#if WITH_PLUGIN_SUPPORT && IS_PROGRAM
+			return true;
+		#endif
+		break;
+	}
+
+	return false;
+}
+
+void FModuleDescriptor::LoadModulesForPhase(ELoadingPhase::Type LoadingPhase, const TArray<FModuleDescriptor>& Modules, TMap<FName, EModuleLoadResult>& ModuleLoadErrors)
+{
+	for(int Idx = 0; Idx < Modules.Num(); Idx++)
+	{
+		const FModuleDescriptor& Descriptor = Modules[Idx];
+
+		// Don't need to do anything if this module is already loaded
+		if( !FModuleManager::Get().IsModuleLoaded( Descriptor.Name ) )
+		{
+			if( LoadingPhase == Descriptor.LoadingPhase && Descriptor.IsLoadedInCurrentConfiguration() )
+			{
+				// @todo plugin: DLL search problems.  Plugins that statically depend on other modules within this plugin may not be found?  Need to test this.
+
+				// NOTE: Loading this module may cause other modules to become loaded, both in the engine or game, or other modules 
+				//       that are part of this project or plugin.  That's totally fine.
+				EModuleLoadResult FailureReason;
+				const TSharedPtr<IModuleInterface>& ModuleInterface = FModuleManager::Get().LoadModuleWithFailureReason( Descriptor.Name, FailureReason );
+				if( ModuleInterface.IsValid() )
+				{
+					// Module loaded OK (or was already loaded.)
+				}
+				else 
+				{
+					// The module failed to load. Note this in the ModuleLoadErrors list.
+					ModuleLoadErrors.Add(Descriptor.Name, FailureReason);
+				}
+			}
+		}
+	}
+}
+
+bool FModuleDescriptor::AreModulesUpToDate(const TArray<FModuleDescriptor>& Modules)
+{
+	for(int Idx = 0; Idx < Modules.Num(); Idx++)
+	{
+		const FModuleDescriptor &Module = Modules[Idx];
+		if (Module.IsCompiledInCurrentConfiguration() && !FModuleManager::Get().IsModuleUpToDate(Module.Name))
+		{
+			return false;
+		}
+	}
+	return true;
+}
+
 #undef LOCTEXT_NAMESPACE
 
diff --git a/Engine/Source/Runtime/Projects/Private/PluginDescriptor.cpp b/Engine/Source/Runtime/Projects/Private/PluginDescriptor.cpp
new file mode 100644
--- /dev/null
+++ b/Engine/Source/Runtime/Projects/Private/PluginDescriptor.cpp
@@ -0,0 +1,94 @@
+// Copyright 1998-2014 Epic Games, Inc. All Rights Reserved.
+
+#include "ProjectsPrivatePCH.h"
+#include "PluginDescriptor.h"
+
+#define LOCTEXT_NAMESPACE "PluginDescriptor"
+
+FPluginDescriptor::FPluginDescriptor()
+{
+	FileVersion = EPluginDescriptorVersion::Latest;
+	Version = 0;
+	bCanContainContent = false;
+	bIsBetaVersion = false;
+}
+
+bool FPluginDescriptor::Load(const FString& FileName, FText& OutFailReason)
+{
+	// Read the file to a string
+	FString FileContents;
+	if (!FFileHelper::LoadFileToString(FileContents, *FileName))
+	{
+		OutFailReason = FText::Format( LOCTEXT("FailedToLoadDescriptorFile", "Failed to open descriptor file '{0}'"), FText::FromString( FileName ) );
+		return false;
+	}
+
+	// Deserialize a JSON object from the string
+	TSharedPtr< FJsonObject > Object;
+	TSharedRef< TJsonReader<> > Reader = TJsonReaderFactory<>::Create(FileContents);
+	if ( !FJsonSerializer::Deserialize(Reader, Object) || !Object.IsValid() )
+	{
+		OutFailReason = FText::Format( LOCTEXT("FailedToReadDescriptorFile", "Failed to read file. {0}"), FText::FromString(Reader->GetErrorMessage()) );
+		return false;
+	}
+
+	// Parse it as a plugin descriptor
+	return Read(*Object.Get(), OutFailReason);
+}
+
+bool FPluginDescriptor::Read(const FJsonObject& Object, FText& OutFailReason)
+{
+	// Read the file version
+	int32 FileVersionInt32;
+	if(!Object.TryGetNumberField(TEXT("FileVersion"), FileVersionInt32))
+	{
+		if(!Object.TryGetNumberField(TEXT("PluginFileVersion"), FileVersionInt32))
+		{
+			OutFailReason = LOCTEXT("InvalidProjectFileVersion", "File does not have a valid 'FileVersion' number.");
+			return false;
+		}
+	}
+
+	// Check that it's within range
+	EPluginDescriptorVersion::Type FileVersion = (EPluginDescriptorVersion::Type)FileVersionInt32;
+	if ( FileVersion <= EPluginDescriptorVersion::Invalid || FileVersion > EPluginDescriptorVersion::Latest )
+	{
+		FText ReadVersionText = FText::FromString( FString::Printf( TEXT( "%d" ), (int32)FileVersion ) );
+		FText LatestVersionText = FText::FromString( FString::Printf( TEXT( "%d" ), (int32)EPluginDescriptorVersion::Latest ) );
+		OutFailReason = FText::Format( LOCTEXT("ProjectFileVersionTooLarge", "File appears to be in a newer version ({0}) of the file format that we can load (max version: {1})."), ReadVersionText, LatestVersionText);
+		return false;
+	}
+
+	// Read the other fields
+	Object.TryGetNumberField(TEXT("Version"), Version);
+	Object.TryGetStringField(TEXT("VersionName"), VersionName);
+	Object.TryGetStringField(TEXT("FriendlyName"), FriendlyName);
+	Object.TryGetStringField(TEXT("Description"), Description);
+
+	if ( !Object.TryGetStringField(TEXT("Category"), Category) )
+	{
+		// Category used to be called CategoryPath in .uplugin files
+		Object.TryGetStringField(TEXT("CategoryPath"), Category);
+	}
+        
+	// Due to a difference in command line parsing between Windows and Mac, we shipped a few Mac samples containing
+	// a category name with escaped quotes. Remove them here to make sure we can list them in the right category.
+	if(Category.Len() >= 2 && Category.StartsWith(TEXT("\"")) && Category.EndsWith(TEXT("\"")))
+	{
+		Category = Category.Mid(1, Category.Len() - 2);
+	}
+
+	Object.TryGetStringField(TEXT("CreatedBy"), CreatedBy);
+	Object.TryGetStringField(TEXT("CreatedByURL"), CreatedByURL);
+
+	if(!FModuleDescriptor::ReadArray(Object, TEXT("Modules"), Modules, OutFailReason))
+	{
+		return false;
+	}
+
+	Object.TryGetBoolField(TEXT("CanContainContent"), bCanContainContent);
+	Object.TryGetBoolField(TEXT("IsBetaVersion"), bIsBetaVersion);
+	return true;
+}
+
+#undef LOCTEXT_NAMESPACE
diff --git a/Engine/Source/Runtime/Projects/Private/PluginManager.cpp b/Engine/Source/Runtime/Projects/Private/PluginManager.cpp
--- a/Engine/Source/Runtime/Projects/Private/PluginManager.cpp
+++ b/Engine/Source/Runtime/Projects/Private/PluginManager.cpp
@@ -2,7 +2,6 @@
 
 #include "ProjectsPrivatePCH.h"
 
-
 DEFINE_LOG_CATEGORY_STATIC( LogPluginManager, Log, All );
 
 #define LOCTEXT_NAMESPACE "PluginManager"
@@ -23,6 +22,36 @@ namespace PluginSystemDefs
 
 }
 
+FPluginInstance::FPluginInstance(const FString& InFileName, const FPluginDescriptor& InDescriptor, EPluginLoadedFrom::Type InLoadedFrom)
+	: Name(FPaths::GetBaseFilename(InFileName))
+	, FileName(InFileName)
+	, Descriptor(InDescriptor)
+	, LoadedFrom(InLoadedFrom)
+	, bEnabled(false)
+{
+}
+
+void FPluginInstance::RegisterPluginMountPoints( const IPluginManager::FRegisterMountPointDelegate& RegisterMountPointDelegate )
+{
+	// If this plugin has content, mount it before loading the first module it contains
+	// @todo plugin: Should we do this even if there is no Content directory?  Probably, so that we can add content in Content Browser!
+	if( Descriptor.bCanContainContent )
+	{
+		if( ensure( RegisterMountPointDelegate.IsBound() ) )
+		{
+			// @todo plugin content: Consider prefixing path with "Plugins" or more specifically, either "EnginePlugins" or "GamePlugins", etc.
+			// @todo plugin content: We may need the full relative path to the plugin under the respective Plugins directory if we want to fully avoid collisions
+			const FString ContentRootPath( FString::Printf( TEXT( "/%s/" ), *Name ) );
+			const FString PluginContentDirectory( FPaths::GetPath(FileName) / TEXT( "Content" ) );
+
+			RegisterMountPointDelegate.Execute( ContentRootPath, PluginContentDirectory );
+		}
+	}
+}
+
+
+
+
 
 
 FPluginManager::FPluginManager()
@@ -58,7 +87,7 @@ void FPluginManager::DiscoverAllPlugins()
 		 * @param	LoadedFrom			Where we're loading these plugins from (game, engine, etc)
 		 * @param	Plugins				The array to be filled in with new plugins including descriptors
 		 */
-		static void FindPluginsRecursively( const FString& PluginsDirectory, const FPluginInfo::ELoadedFrom::Type LoadedFrom, TArray< TSharedRef<FPlugin> >& Plugins )
+		static void FindPluginsRecursively( const FString& PluginsDirectory, const EPluginLoadedFrom::Type LoadedFrom, TArray< TSharedRef<FPluginInstance> >& Plugins )
 		{
 			// NOTE: The logic in this function generally matches that of the C# code for FindPluginsRecursively
 			//       in UnrealBuildTool.  These routines should be kept in sync.
@@ -159,15 +188,13 @@ void FPluginManager::DiscoverAllPlugins()
 				if( !PluginDescriptorFilename.IsEmpty() )
 				{
 					// Found a plugin directory!  No need to recurse any further.
+					FPluginDescriptor Descriptor;
 
-					// Load the descriptor file for this plugin into a new plugin info entry
-					TSharedRef< FPlugin > NewPlugin = MakeShareable( new FPlugin() );
-
+					// Load the descriptor
 					FText FailureReason;
-					const bool bLoadedSuccessfully = NewPlugin->LoadFromFile(PluginDescriptorFilename, FailureReason);
-					if( bLoadedSuccessfully )
+					if(Descriptor.Load(PluginDescriptorFilename, FailureReason))
 					{
-						NewPlugin->SetLoadedFrom(LoadedFrom);
+						TSharedRef< FPluginInstance > NewPlugin = MakeShareable( new FPluginInstance(PluginDescriptorFilename, Descriptor, LoadedFrom) );
 
 						if (FPaths::IsProjectFilePathSet())
 						{
@@ -175,7 +202,7 @@ void FPluginManager::DiscoverAllPlugins()
 							if (PluginDescriptorFilename.StartsWith(GameProjectFolder))
 							{
 								FString PluginBinariesFolder = FPaths::Combine(*FPaths::GetPath(PluginDescriptorFilename), TEXT("Binaries"), FPlatformProcess::GetBinariesSubdirectory());
-								FModuleManager::Get().AddBinariesDirectory(*PluginBinariesFolder, (LoadedFrom == FPluginInfo::ELoadedFrom::GameProject));
+								FModuleManager::Get().AddBinariesDirectory(*PluginBinariesFolder, (LoadedFrom == EPluginLoadedFrom::GameProject));
 							}
 						}
 
@@ -205,7 +232,7 @@ void FPluginManager::DiscoverAllPlugins()
 		 * @param	LoadedFrom			Where we're loading these plugins from (game, engine, etc)
 		 * @param	Plugins				The array to be filled in with loaded plugin descriptors
 		 */
-		static void FindPluginsIn( const FString& PluginsDirectory, const FPluginInfo::ELoadedFrom::Type LoadedFrom, TArray< TSharedRef<FPlugin> >& Plugins )
+		static void FindPluginsIn( const FString& PluginsDirectory, const EPluginLoadedFrom::Type LoadedFrom, TArray< TSharedRef<FPluginInstance> >& Plugins )
 		{
 			// Make sure the directory even exists
 			if( FPlatformFileManager::Get().GetPlatformFile().DirectoryExists( *PluginsDirectory ) )
@@ -217,16 +244,16 @@ void FPluginManager::DiscoverAllPlugins()
 
 
 	{
-		TArray< TSharedRef<FPlugin> > Plugins;
+		TArray< TSharedRef<FPluginInstance> > Plugins;
 
 #if (WITH_ENGINE && !IS_PROGRAM) || WITH_PLUGIN_SUPPORT
 		// Find "built-in" plugins.  That is, plugins situated right within the Engine directory.
-		Local::FindPluginsIn( FPaths::EnginePluginsDir(), FPluginInfo::ELoadedFrom::Engine, Plugins );
+		Local::FindPluginsIn( FPaths::EnginePluginsDir(), EPluginLoadedFrom::Engine, Plugins );
 
 		// Find plugins in the game project directory (<MyGameProject>/Plugins)
 		if( FApp::HasGameName() )
 		{
-			Local::FindPluginsIn( FPaths::GamePluginsDir(), FPluginInfo::ELoadedFrom::GameProject, Plugins );
+			Local::FindPluginsIn( FPaths::GamePluginsDir(), EPluginLoadedFrom::GameProject, Plugins );
 		}
 #endif		// (WITH_ENGINE && !IS_PROGRAM) || WITH_PLUGIN_SUPPORT
 
@@ -236,15 +263,14 @@ void FPluginManager::DiscoverAllPlugins()
 		AllPlugins = Plugins;
 		for( auto PluginIt( Plugins.CreateConstIterator() ); PluginIt; ++PluginIt )
 		{
-			const TSharedRef<FPlugin>& Plugin = *PluginIt;
-			const FPluginInfo& PluginInfo = Plugin->GetPluginInfo();
+			const TSharedRef<FPluginInstance>& Plugin = *PluginIt;
 
 			// Add the plugin binaries directory
-			const FString PluginBinariesPath = FPaths::Combine(*FPaths::GetPath(PluginInfo.LoadedFromFile), TEXT("Binaries"), FPlatformProcess::GetBinariesSubdirectory());
-			FModuleManager::Get().AddBinariesDirectory(*PluginBinariesPath, PluginInfo.LoadedFrom == FPluginInfo::ELoadedFrom::GameProject);
+			const FString PluginBinariesPath = FPaths::Combine(*FPaths::GetPath(Plugin->FileName), TEXT("Binaries"), FPlatformProcess::GetBinariesSubdirectory());
+			FModuleManager::Get().AddBinariesDirectory(*PluginBinariesPath, Plugin->LoadedFrom == EPluginLoadedFrom::GameProject);
 
 			// Make sure to tell the module manager about any of our plugin's modules, so it will know where on disk to find them.
-			for( auto ModuleInfoIt( PluginInfo.Modules.CreateConstIterator() ); ModuleInfoIt; ++ModuleInfoIt )
+			for( auto ModuleInfoIt( Plugin->Descriptor.Modules.CreateConstIterator() ); ModuleInfoIt; ++ModuleInfoIt )
 			{
 				const FModuleDescriptor& ModuleInfo = *ModuleInfoIt;
 
@@ -266,20 +292,20 @@ void FPluginManager::EnablePluginsThatAreConfiguredToBeEnabled()
 	TSet< FString > AllEnabledPlugins;
 	AllEnabledPlugins.Append( MoveTemp(EnabledPluginNames) );
 
-	for( const TSharedRef< FPlugin > Plugin : AllPlugins )
+	for( const TSharedRef< FPluginInstance > Plugin : AllPlugins )
 	{
-		if ( AllEnabledPlugins.Contains(Plugin->GetPluginName()) )
+		if ( AllEnabledPlugins.Contains(Plugin->Name) )
 		{
-			Plugin->SetEnabled(true);
+			Plugin->bEnabled = true;
 		}
 	}
 }
 
 void FPluginManager::RegisterEnabledPluginMountPoints()
 {
-	for( const TSharedRef< FPlugin > Plugin : AllPlugins )
+	for( const TSharedRef< FPluginInstance > Plugin : AllPlugins )
 	{
-		if ( Plugin->IsEnabled() )
+		if ( Plugin->bEnabled )
 		{
 			Plugin->RegisterPluginMountPoints( RegisterMountPointDelegate );
 		}
@@ -296,12 +322,12 @@ void FPluginManager::LoadModulesForEnabledPlugins( const ELoadingPhase::Type Loa
 	}
 
 	// Load plugins!
-	for( const TSharedRef< FPlugin > Plugin : AllPlugins )
+	for( const TSharedRef< FPluginInstance > Plugin : AllPlugins )
 	{
-		if ( Plugin->IsEnabled() )
+		if ( Plugin->bEnabled )
 		{
 			TMap<FName, EModuleLoadResult> ModuleLoadFailures;
-			Plugin->LoadModules( LoadingPhase, ModuleLoadFailures );
+			FModuleDescriptor::LoadModulesForPhase(LoadingPhase, Plugin->Descriptor.Modules, ModuleLoadFailures);
 
 			FText FailureMessage;
 			for( auto FailureIt( ModuleLoadFailures.CreateConstIterator() ); FailureIt; ++FailureIt )
@@ -311,7 +337,7 @@ void FPluginManager::LoadModulesForEnabledPlugins( const ELoadingPhase::Type Loa
 
 				if( FailureReason != EModuleLoadResult::Success )
 				{
-					const FText PluginNameText = FText::FromString(Plugin->GetPluginName());
+					const FText PluginNameText = FText::FromString(Plugin->Name);
 					const FText TextModuleName = FText::FromName(FailureIt.Key());
 
 					if ( FailureReason == EModuleLoadResult::FileNotFound )
@@ -368,10 +394,10 @@ bool FPluginManager::AreEnabledPluginModulesUpToDate()
 		bEarliestPhaseProcessed = true;
 	}
 
-	for (TArray< TSharedRef< FPlugin > >::TConstIterator Iter(AllPlugins); Iter; ++Iter)
+	for (TArray< TSharedRef< FPluginInstance > >::TConstIterator Iter(AllPlugins); Iter; ++Iter)
 	{
-		const TSharedRef< FPlugin > &Plugin = *Iter;
-		if (Plugin->IsEnabled() && !Plugin->AreModulesUpToDate())
+		const TSharedRef< FPluginInstance > &Plugin = *Iter;
+		if (Plugin->bEnabled && !FModuleDescriptor::AreModulesUpToDate(Plugin->Descriptor.Modules))
 		{
 			return false;
 		}
@@ -478,80 +504,17 @@ IPluginManager& IPluginManager::Get()
 	return *PluginManager;
 }
 
-
-
-FPlugin::FPlugin()
-	: bPluginEnabled( false )
-{
-}
-
-void FPlugin::RegisterPluginMountPoints( const IPluginManager::FRegisterMountPointDelegate& RegisterMountPointDelegate )
-{
-	// If this plugin has content, mount it before loading the first module it contains
-	// @todo plugin: Should we do this even if there is no Content directory?  Probably, so that we can add content in Content Browser!
-	if( PluginInfo.bCanContainContent )
-	{
-		if( ensure( RegisterMountPointDelegate.IsBound() ) )
-		{
-			// @todo plugin content: Consider prefixing path with "Plugins" or more specifically, either "EnginePlugins" or "GamePlugins", etc.
-			// @todo plugin content: We may need the full relative path to the plugin under the respective Plugins directory if we want to fully avoid collisions
-			const FString ContentRootPath( FString::Printf( TEXT( "/%s/" ), *PluginInfo.Name ) );
-			const FString PluginContentDirectory( FPaths::GetPath(PluginInfo.LoadedFromFile) / TEXT( "Content" ) );
-
-			RegisterMountPointDelegate.Execute( ContentRootPath, PluginContentDirectory );
-		}
-	}
-}
-
-
-const FString& FPlugin::GetPluginName() const
-{
-	return PluginInfo.Name;
-}
-
-
-void FPlugin::SetEnabled(bool bNewPluginEnabled)
-{
-	bPluginEnabled = bNewPluginEnabled;
-}
-
-
-bool FPlugin::IsEnabled() const
-{
-	return bPluginEnabled;
-}
-
-void FPlugin::SetLoadedFrom(FPluginInfo::ELoadedFrom::Type NewLoadedFrom)
-{
-	PluginInfo.LoadedFrom = NewLoadedFrom;
-}
-
-bool FPlugin::PerformAdditionalDeserialization(const TSharedRef< FJsonObject >& FileObject)
-{
-	ReadBoolFromJSON(FileObject, TEXT("CanContainContent"), PluginInfo.bCanContainContent);
-	ReadBoolFromJSON(FileObject, TEXT("IsBetaVersion"), PluginInfo.bIsBetaVersion);
-
-	return true;
-}
-
-void FPlugin::PerformAdditionalSerialization(const TSharedRef< TJsonWriter<> >& Writer) const
-{
-	Writer->WriteValue(TEXT("CanContainContent"), PluginInfo.bCanContainContent);
-	Writer->WriteValue(TEXT("IsBetaVersion"), PluginInfo.bIsBetaVersion);
-}
-
-
 TArray< FPluginStatus > FPluginManager::QueryStatusForAllPlugins() const
 {
 	TArray< FPluginStatus > PluginStatuses;
 
 	for( auto PluginIt( AllPlugins.CreateConstIterator() ); PluginIt; ++PluginIt )
 	{
-		const TSharedRef< FPlugin >& Plugin = *PluginIt;
-		const auto& PluginInfo = Plugin->GetPluginInfo();
+		const TSharedRef< FPluginInstance >& Plugin = *PluginIt;
+		const FPluginDescriptor& PluginInfo = Plugin->Descriptor;
 		
 		FPluginStatus PluginStatus;
-		PluginStatus.Name = PluginInfo.Name;
+		PluginStatus.Name = Plugin->Name;
 		PluginStatus.FriendlyName = PluginInfo.FriendlyName;
 		PluginStatus.Version = PluginInfo.Version;
 		PluginStatus.VersionName = PluginInfo.VersionName;
@@ -559,12 +522,12 @@ TArray< FPluginStatus > FPluginManager::QueryStatusForAllPlugins() const
 		PluginStatus.CreatedBy = PluginInfo.CreatedBy;
 		PluginStatus.CreatedByURL = PluginInfo.CreatedByURL;
 		PluginStatus.CategoryPath = PluginInfo.Category;
-		PluginStatus.bIsEnabled = Plugin->IsEnabled();
-		PluginStatus.bIsBuiltIn = ( PluginInfo.LoadedFrom == FPluginInfo::ELoadedFrom::Engine );
+		PluginStatus.bIsEnabled = Plugin->bEnabled;
+		PluginStatus.bIsBuiltIn = ( Plugin->LoadedFrom == EPluginLoadedFrom::Engine );
 		PluginStatus.bIsBetaVersion = PluginInfo.bIsBetaVersion;
 
 		// @todo plugedit: Maybe we should do the FileExists check ONCE at plugin load time and not at query time
-		const FString Icon128FilePath = FPaths::GetPath(PluginInfo.LoadedFromFile) / PluginSystemDefs::RelativeIcon128FilePath;
+		const FString Icon128FilePath = FPaths::GetPath(Plugin->FileName) / PluginSystemDefs::RelativeIcon128FilePath;
 		if( FPlatformFileManager::Get().GetPlatformFile().FileExists( *Icon128FilePath ) )
 		{
 			PluginStatus.Icon128FilePath = Icon128FilePath;
@@ -580,11 +543,11 @@ void FPluginManager::SetPluginEnabled( const FString& PluginName, bool bEnabled
 {
 	for( auto PluginIt( AllPlugins.CreateConstIterator() ); PluginIt; ++PluginIt )
 	{
-		const TSharedRef< FPlugin >& Plugin = *PluginIt;
+		const TSharedRef< FPluginInstance >& Plugin = *PluginIt;
 
-		if ( Plugin->GetPluginName() == PluginName )
+		if ( Plugin->Name == PluginName )
 		{
-			Plugin->SetEnabled(bEnabled);
+			Plugin->bEnabled = bEnabled;
 			ConfigurePluginToBeEnabled(PluginName, bEnabled);
 			break;
 		}
diff --git a/Engine/Source/Runtime/Projects/Private/PluginManager.h b/Engine/Source/Runtime/Projects/Private/PluginManager.h
--- a/Engine/Source/Runtime/Projects/Private/PluginManager.h
+++ b/Engine/Source/Runtime/Projects/Private/PluginManager.h
@@ -2,64 +2,48 @@
 
 #pragma once
 
+#include "PluginDescriptor.h"
 
 /**
- * Info about a single plugin
+ * Enum for where a plugin is loaded from
  */
-class FPluginInfo : public FProjectOrPluginInfo
+struct EPluginLoadedFrom
 {
-public:
-	struct ELoadedFrom
+	enum Type
 	{
-		enum Type
-		{
-			/** Plugin is built-in to the engine */
-			Engine,
-
-			/** Project-specific plugin, stored within a game project directory */
-			GameProject
-		};
-	};
-			
-
-	/** Where does this plugin live? */
-	ELoadedFrom::Type LoadedFrom;
-
+		/** Plugin is built-in to the engine */
+		Engine,
 
-	//////////////////////////////////////////////////////////////////////////////////////////////////////
-	// Everything Below this line is serialized. If you add or remove properties, you must update
-	// FPlugin::PerformAdditionalSerialization and FPlugin::PerformAdditionalDeserialization
-	//////////////////////////////////////////////////////////////////////////////////////////////////////
-
-	/** Can this plugin contain content? */
-	bool bCanContainContent;
-
-	/** Marks the plugin as beta in the UI */
-	bool bIsBetaVersion;
-
-	/** Plugin info constructor */
-	FPluginInfo()
-		: FProjectOrPluginInfo()
-		, LoadedFrom( ELoadedFrom::Engine )
-		, bCanContainContent( false )
-		, bIsBetaVersion( false )
-	{
-	}
+		/** Project-specific plugin, stored within a game project directory */
+		GameProject
+	};
 };
 
-
 /**
  * Instance of a plugin in memory
  */
-class FPlugin : public FProjectOrPlugin
+class FPluginInstance
 {
-
 public:
+	/** The name of the plugin */
+	FString Name;
+
+	/** The filename that the plugin was loaded from */
+	FString FileName;
+
+	/** The plugin's settings */
+	FPluginDescriptor Descriptor;
+
+	/** Where does this plugin live? */
+	EPluginLoadedFrom::Type LoadedFrom;
+
+	/** True if the plugin is marked as enabled */
+	bool bEnabled;
 
 	/**
 	 * FPlugin constructor
 	 */
-	FPlugin();
+	FPluginInstance(const FString &FileName, const FPluginDescriptor& InDescriptor, EPluginLoadedFrom::Type InLoadedFrom);
 
 	/**
 	 * Registers the content mount points for this delegate. This allows content to be created in the root folder for this plugin.
@@ -67,60 +51,20 @@ public:
 	 * @param	RegisterMountPointDelegate	Delegate for mounting content paths.  Bound by FPackageName code in CoreUObject, so that we can access content path mounting functionality from Core.
 	 */
 	void RegisterPluginMountPoints( const IPluginManager::FRegisterMountPointDelegate& RegisterMountPointDelegate );
-
-	/** @return Exposes access to the plugin's descriptor */
-	const FPluginInfo& GetPluginInfo()
-	{
-		return PluginInfo;
-	}
-
-	/** Returns this plugin's name */
-	const FString& GetPluginName() const;
-
-	/** Sets the "enabled" state for this plugin */
-	void SetEnabled(bool bNewPluginEnabled);
-
-	/** Returns true if this plugin is currently marked as "enabled" */
-	bool IsEnabled() const;
-
-	/** Tells the descriptor where it was loaded from */
-	void SetLoadedFrom(FPluginInfo::ELoadedFrom::Type NewLoadedFrom);
-
-protected:
-	virtual FProjectOrPluginInfo& GetProjectOrPluginInfo() override	{ return PluginInfo; }
-	virtual const FProjectOrPluginInfo& GetProjectOrPluginInfo() const override { return PluginInfo; }
-
-	virtual bool PerformAdditionalDeserialization(const TSharedRef< FJsonObject >& FileObject) override;
-	virtual void PerformAdditionalSerialization(const TSharedRef< TJsonWriter<> >& Writer) const override;
-
-private:
-
-	/** The descriptor of this plugin that we loaded from disk */
-	FPluginInfo PluginInfo;
-
-	/** True if the plugin is marked as enabled */
-	bool bPluginEnabled;
 };
 
-
-
 /**
- * ProjectAndPluginManager manages available code and content extensions (both loaded and not loaded.)
+ * FPluginManager manages available code and content extensions (both loaded and not loaded.)
  */
 class FPluginManager : public IPluginManager
 {
-
 public:
-
 	/** Constructor */
 	FPluginManager();
 
 	/** Destructor */
 	~FPluginManager();
 
-
-public:
-
 	/** IPluginManager interface */
 	virtual void LoadModulesForEnabledPlugins( const ELoadingPhase::Type LoadingPhase ) override;
 	virtual void SetRegisterMountPointDelegate( const FRegisterMountPointDelegate& Delegate ) override;
@@ -151,10 +95,10 @@ private:
 
 private:
 	/** All of the plugins that we know about */
-	TArray< TSharedRef< FPlugin > > AllPlugins;
+	TArray< TSharedRef< FPluginInstance > > AllPlugins;
 
 	/** Map from module name to plugin containing that module */
-	TMap< FName, TSharedRef< FPlugin > > ModulePluginMap;
+	TMap< FName, TSharedRef< FPluginInstance > > ModulePluginMap;
 
 	/** Delegate for mounting content paths.  Bound by FPackageName code in CoreUObject, so that we can access
 	    content path mounting functionality from Core. */
diff --git a/Engine/Source/Runtime/Projects/Private/PluginManagerShared.cpp b/Engine/Source/Runtime/Projects/Private/PluginManagerShared.cpp
deleted file mode 100644
diff --git a/Engine/Source/Runtime/Projects/Private/PluginManagerShared.h b/Engine/Source/Runtime/Projects/Private/PluginManagerShared.h
deleted file mode 100644
diff --git a/Engine/Source/Runtime/Projects/Private/ProjectDescriptor.cpp b/Engine/Source/Runtime/Projects/Private/ProjectDescriptor.cpp
new file mode 100644
--- /dev/null
+++ b/Engine/Source/Runtime/Projects/Private/ProjectDescriptor.cpp
@@ -0,0 +1,165 @@
+// Copyright 1998-2014 Epic Games, Inc. All Rights Reserved.
+
+#include "ProjectsPrivatePCH.h"
+#include "ProjectDescriptor.h"
+
+#define LOCTEXT_NAMESPACE "ProjectDescriptor"
+
+FProjectDescriptor::FProjectDescriptor()
+{
+	FileVersion = EProjectDescriptorVersion::Latest;
+	EpicSampleNameHash = 0;
+}
+
+void FProjectDescriptor::Sign(const FString& FilePath)
+{
+	EpicSampleNameHash = GetTypeHash(FPaths::GetCleanFilename(FilePath));
+}
+
+bool FProjectDescriptor::IsSigned(const FString& FilePath) const
+{
+	return EpicSampleNameHash == GetTypeHash(FPaths::GetCleanFilename(FilePath));
+}
+
+void FProjectDescriptor::UpdateSupportedTargetPlatforms(const FName& InPlatformName, bool bIsSupported)
+{
+	if ( bIsSupported )
+	{
+		TargetPlatforms.AddUnique(InPlatformName);
+	}
+	else
+	{
+		TargetPlatforms.Remove(InPlatformName);
+	}
+}
+
+bool FProjectDescriptor::Load(const FString& FileName, FText& OutFailReason)
+{
+	// Read the file to a string
+	FString FileContents;
+	if (!FFileHelper::LoadFileToString(FileContents, *FileName))
+	{
+		OutFailReason = FText::Format(LOCTEXT("FailedToLoadDescriptorFile", "Failed to open descriptor file '{0}'"), FText::FromString(FileName));
+		return false;
+	}
+
+	// Deserialize a JSON object from the string
+	TSharedPtr< FJsonObject > Object;
+	TSharedRef< TJsonReader<> > Reader = TJsonReaderFactory<>::Create(FileContents);
+	if ( !FJsonSerializer::Deserialize(Reader, Object) || !Object.IsValid() )
+	{
+		OutFailReason = FText::Format(LOCTEXT("FailedToReadDescriptorFile", "Failed to read file. {1}"), FText::FromString(Reader->GetErrorMessage()));
+		return false;
+	}
+
+	// Parse it as a project descriptor
+	return Read(*Object.Get(), OutFailReason);
+}
+
+bool FProjectDescriptor::Read(const FJsonObject& Object, FText& OutFailReason)
+{
+	// Read the file version
+	int32 FileVersionInt32;
+	if(!Object.TryGetNumberField(TEXT("FileVersion"), FileVersionInt32))
+	{
+		if(!Object.TryGetNumberField(TEXT("ProjectFileVersion"), FileVersionInt32))
+		{
+			OutFailReason = LOCTEXT("InvalidProjectFileVersion", "File does not have a valid 'FileVersion' number.");
+			return false;
+		}
+	}
+
+	// Check that it's within range
+	FileVersion = (EProjectDescriptorVersion::Type)FileVersionInt32;
+	if ( FileVersion <= EProjectDescriptorVersion::Invalid || FileVersion > EProjectDescriptorVersion::Latest )
+	{
+		FText ReadVersionText = FText::FromString( FString::Printf( TEXT( "%d" ), (int32)FileVersion ) );
+		FText LatestVersionText = FText::FromString( FString::Printf( TEXT( "%d" ), (int32)EProjectDescriptorVersion::Latest ) );
+		OutFailReason = FText::Format( LOCTEXT("ProjectFileVersionTooLarge", "File appears to be in a newer version ({0}) of the file format that we can load (max version: {1})."), ReadVersionText, LatestVersionText);
+		return false;
+	}
+
+	// Read simple fields
+	Object.TryGetStringField(TEXT("EngineAssociation"), EngineAssociation);
+	Object.TryGetStringField(TEXT("Category"), Category);
+	Object.TryGetStringField(TEXT("Description"), Description);
+
+	// Read the modules
+	if(!FModuleDescriptor::ReadArray(Object, TEXT("Modules"), Modules, OutFailReason))
+	{
+		return false;
+	}
+
+	// Read the target platforms
+	const TArray< TSharedPtr<FJsonValue> > *TargetPlatformsValue;
+	if(Object.TryGetArrayField(TEXT("TargetPlatforms"), TargetPlatformsValue))
+	{
+		for(int32 Idx = 0; Idx < TargetPlatformsValue->Num(); Idx++)
+		{
+			FString TargetPlatform;
+			if((*TargetPlatformsValue)[Idx]->TryGetString(TargetPlatform))
+			{
+				TargetPlatforms.Add(*TargetPlatform);
+			}
+		}
+	}
+
+	// Get the sample name hash
+	Object.TryGetNumberField(TEXT("EpicSampleNameHash"), EpicSampleNameHash);
+	return true;
+}
+
+bool FProjectDescriptor::Save(const FString& FileName, FText& OutFailReason)
+{
+	// Write the contents of the descriptor to a string. Make sure the writer is destroyed so that the contents are flushed to the string.
+	FString Text;
+	TSharedRef< TJsonWriter<> > Writer = TJsonWriterFactory<>::Create(&Text);
+	Write(Writer.Get());
+	Writer->Close();
+
+	// Save it to a file
+	if ( FFileHelper::SaveStringToFile(Text, *FileName) )
+	{
+		return true;
+	}
+	else
+	{
+		OutFailReason = FText::Format( LOCTEXT("FailedToWriteOutputFile", "Failed to write output file '{0}'. Perhaps the file is Read-Only?"), FText::FromString(FileName) );
+		return false;
+	}
+}
+
+void FProjectDescriptor::Write(TJsonWriter<>& Writer) const
+{
+	Writer.WriteObjectStart();
+
+	// Write all the simple fields
+	Writer.WriteValue(TEXT("FileVersion"), EProjectDescriptorVersion::Latest);
+	Writer.WriteValue(TEXT("EngineAssociation"), EngineAssociation);
+	Writer.WriteValue(TEXT("Category"), Category);
+	Writer.WriteValue(TEXT("Description"), Description);
+
+	// Write the module list
+	FModuleDescriptor::WriteArray(Writer, TEXT("Modules"), Modules);
+
+	// Write the target platforms
+	if ( TargetPlatforms.Num() > 0 )
+	{
+		Writer.WriteArrayStart(TEXT("TargetPlatforms"));
+		for(int Idx = 0; Idx < TargetPlatforms.Num(); Idx++)
+		{
+			Writer.WriteValue(TargetPlatforms[Idx].ToString());
+		}
+		Writer.WriteArrayEnd();
+	}
+
+	// If it's a signed sample, write the name hash
+	if(EpicSampleNameHash != 0)
+	{
+		Writer.WriteValue(TEXT("EpicSampleNameHash"), FString::Printf(TEXT("%u"), EpicSampleNameHash));
+	}
+
+	Writer.WriteObjectEnd();
+}
+
+#undef LOCTEXT_NAMESPACE
diff --git a/Engine/Source/Runtime/Projects/Private/ProjectManager.cpp b/Engine/Source/Runtime/Projects/Private/ProjectManager.cpp
--- a/Engine/Source/Runtime/Projects/Private/ProjectManager.cpp
+++ b/Engine/Source/Runtime/Projects/Private/ProjectManager.cpp
@@ -2,101 +2,23 @@
 
 #include "ProjectsPrivatePCH.h"
 
-
 DEFINE_LOG_CATEGORY_STATIC( LogProjectManager, Log, All );
 
 #define LOCTEXT_NAMESPACE "ProjectManager"
 
-
-FProject::FProject()
-{
-}
-
-FProject::FProject( const FProjectInfo& InitProjectInfo )
-	: ProjectInfo(InitProjectInfo)
-{
-}
-
-bool FProject::IsSignedSampleProject(const FString& FilePath) const
-{
-	return ProjectInfo.EpicSampleNameHash == GetTypeHash(FPaths::GetCleanFilename(FilePath));
-}
-
-void FProject::SignSampleProject(const FString& FilePath, const FString& Category)
-{
-	ProjectInfo.EpicSampleNameHash = GetTypeHash(FPaths::GetCleanFilename(FilePath));
-	ProjectInfo.Category = Category;
-}
-
-void FProject::UpdateSupportedTargetPlatforms(const FName& InPlatformName, const bool bIsSupported)
-{
-	if ( bIsSupported )
-	{
-		ProjectInfo.TargetPlatforms.AddUnique(InPlatformName);
-	}
-	else
-	{
-		ProjectInfo.TargetPlatforms.Remove(InPlatformName);
-	}
-}
-
-void FProject::ClearSupportedTargetPlatforms()
-{
-	ProjectInfo.TargetPlatforms.Empty();
-}
-
-bool FProject::PerformAdditionalDeserialization(const TSharedRef< FJsonObject >& FileObject)
-{
-	ReadNumberFromJSON(FileObject, TEXT("EpicSampleNameHash"), ProjectInfo.EpicSampleNameHash);
-
-	ProjectInfo.TargetPlatforms.Empty();
-	if ( FileObject->HasField(TEXT("TargetPlatforms")) )
-	{
-		const TSharedPtr<FJsonValue>& TargetPlatformsValue = FileObject->GetField<EJson::Array>(TEXT("TargetPlatforms"));
-		const TArray<TSharedPtr<FJsonValue>>& TargetPlatformsArray = TargetPlatformsValue->AsArray();
-		for ( const auto& TargetPlatformsEntry : TargetPlatformsArray )
-		{
-			const FString PlatformName = TargetPlatformsEntry->AsString();
-			ProjectInfo.TargetPlatforms.Add(*PlatformName);
-		}
-	}
-
-	return true;
-}
-
-void FProject::PerformAdditionalSerialization(const TSharedRef< TJsonWriter<> >& Writer) const
-{
-	Writer->WriteValue(TEXT("EpicSampleNameHash"), FString::Printf(TEXT("%u"), ProjectInfo.EpicSampleNameHash));
-
-	if ( ProjectInfo.TargetPlatforms.Num() > 0 )
-	{
-		Writer->WriteArrayStart(TEXT("TargetPlatforms"));
-		for ( const FName& PlatformName : ProjectInfo.TargetPlatforms )
-		{
-			Writer->WriteValue(PlatformName.ToString());
-		}
-		Writer->WriteArrayEnd();
-	}
-}
-
-
-
-
-
 FProjectManager::FProjectManager()
 {
-
 }
 
 bool FProjectManager::LoadProjectFile( const FString& InProjectFile )
 {
+	// Try to load the descriptor
 	FText FailureReason;
-	TSharedRef<FProject> NewProject = MakeShareable( new FProject() );
-	if ( NewProject->LoadFromFile(InProjectFile, FailureReason) )
+	TSharedPtr<FProjectDescriptor> Descriptor = MakeShareable(new FProjectDescriptor());
+	if(Descriptor->Load(InProjectFile, FailureReason))
 	{
-		// Load successful. Set the loaded project file pointer.
-		CurrentlyLoadedProject = NewProject;
-
+		// Create the project
+		CurrentlyLoadedProject = Descriptor;
 		return true;
 	}
 	
@@ -120,7 +42,7 @@ bool FProjectManager::LoadModulesForProject( const ELoadingPhase::Type LoadingPh
 	if ( CurrentlyLoadedProject.IsValid() )
 	{
 		TMap<FName, EModuleLoadResult> ModuleLoadFailures;
-		CurrentlyLoadedProject->LoadModules(LoadingPhase, ModuleLoadFailures);
+		FModuleDescriptor::LoadModulesForPhase(LoadingPhase, CurrentlyLoadedProject->Modules, ModuleLoadFailures);
 
 		if ( ModuleLoadFailures.Num() > 0 )
 		{
@@ -170,7 +92,7 @@ bool FProjectManager::LoadModulesForProject( const ELoadingPhase::Type LoadingPh
 
 bool FProjectManager::AreProjectModulesUpToDate()
 {
-	return !CurrentlyLoadedProject.IsValid() || CurrentlyLoadedProject->AreModulesUpToDate();
+	return !CurrentlyLoadedProject.IsValid() || FModuleDescriptor::AreModulesUpToDate(CurrentlyLoadedProject->Modules);
 }
 
 const FString& FProjectManager::GetAutoLoadProjectFileName()
@@ -187,61 +109,41 @@ const FString& FProjectManager::NonStaticGetProjectFileExtension()
 
 bool FProjectManager::GenerateNewProjectFile(const FString& NewProjectFilename, const TArray<FString>& StartupModuleNames, const FString& EngineIdentifier, FText& OutFailReason)
 {
-	TSharedRef<FProject> NewProject = MakeShareable( new FProject() );
-	NewProject->UpdateVersionToCurrent(EngineIdentifier);
-	NewProject->ReplaceModulesInProject(&StartupModuleNames);
+	FProjectDescriptor Descriptor;
+	Descriptor.EngineAssociation = EngineIdentifier;
 
-	const FString& FileContents = NewProject->SerializeToJSON();
-	if ( FFileHelper::SaveStringToFile(FileContents, *NewProjectFilename) )
-	{
-		return true;
-	}
-	else
+	for(int32 Idx = 0; Idx < StartupModuleNames.Num(); Idx++)
 	{
-		OutFailReason = FText::Format( LOCTEXT("FailedToWriteOutputFile", "Failed to write output file '{0}'. Perhaps the file is Read-Only?"), FText::FromString(NewProjectFilename) );
-		return false;
+		Descriptor.Modules.Add(FModuleDescriptor(*StartupModuleNames[Idx]));
 	}
+
+	return Descriptor.Save(NewProjectFilename, OutFailReason);
 }
 
 bool FProjectManager::DuplicateProjectFile(const FString& SourceProjectFilename, const FString& NewProjectFilename, const FString& EngineIdentifier, FText& OutFailReason)
 {
 	// Load the source project
-	TSharedRef<FProject> SourceProject = MakeShareable( new FProject() );
-	if ( !SourceProject->LoadFromFile(SourceProjectFilename, OutFailReason) )
+	FProjectDescriptor Project;
+	if(!Project.Load(SourceProjectFilename, OutFailReason))
 	{
 		return false;
 	}
 
-	// Duplicate the project info
-	FProjectInfo ProjectInfo = SourceProject->GetProjectInfo();
-
-	// Clear the sample hash
-	ProjectInfo.EpicSampleNameHash = 0;
+	// Update it to current
+	Project.EngineAssociation = EngineIdentifier;
+	Project.EpicSampleNameHash = 0;
 
 	// Fix up module names
 	const FString BaseSourceName = FPaths::GetBaseFilename(SourceProjectFilename);
 	const FString BaseNewName = FPaths::GetBaseFilename(NewProjectFilename);
-	for ( auto ModuleIt = ProjectInfo.Modules.CreateIterator(); ModuleIt; ++ModuleIt )
+	for ( auto ModuleIt = Project.Modules.CreateIterator(); ModuleIt; ++ModuleIt )
 	{
 		FModuleDescriptor& ModuleInfo = *ModuleIt;
 		ModuleInfo.Name = FName(*ModuleInfo.Name.ToString().Replace(*BaseSourceName, *BaseNewName));
 	}
 
-	// Create new project, update version numbers (no need to replace modules here)
-	TSharedRef<FProject> NewProject = MakeShareable( new FProject(ProjectInfo) );
-	NewProject->UpdateVersionToCurrent(EngineIdentifier);
-
-	// Serialize and write to disk
-	const FString& FileContents = NewProject->SerializeToJSON();
-	if ( FFileHelper::SaveStringToFile(FileContents, *NewProjectFilename) )
-	{
-		return true;
-	}
-	else
-	{
-		OutFailReason = FText::Format( LOCTEXT("FailedToWriteOutputFile", "Failed to write output file '{0}'. Perhaps the file is Read-Only?"), FText::FromString(NewProjectFilename) );
-		return false;
-	}
+	// Save it to disk
+	return Project.Save(NewProjectFilename, OutFailReason);
 }
 
 bool FProjectManager::UpdateLoadedProjectFileToCurrent(const TArray<FString>* StartupModuleNames, const FString& EngineIdentifier, FText& OutFailReason)
@@ -252,57 +154,45 @@ bool FProjectManager::UpdateLoadedProjectFileToCurrent(const TArray<FString>* St
 	}
 
 	// Freshen version information
-	CurrentlyLoadedProject->UpdateVersionToCurrent(EngineIdentifier);
+	CurrentlyLoadedProject->EngineAssociation = EngineIdentifier;
 
 	// Replace the modules names, if specified
-	CurrentlyLoadedProject->ReplaceModulesInProject(StartupModuleNames);
-
-	// Update file on disk
-	const FString& FileContents = CurrentlyLoadedProject->SerializeToJSON();
-	if ( FFileHelper::SaveStringToFile(FileContents, *FPaths::GetProjectFilePath()) )
-	{
-		return true;
-	}
-	else
+	if(StartupModuleNames != NULL)
 	{
-		// We failed to generate the file. Could be read only.
-		OutFailReason = FText::Format( LOCTEXT("FailedToWriteOutputFile", "Failed to write output file '{0}'. Perhaps the file is Read-Only?"), FText::FromString(FPaths::GetProjectFilePath()) );
-		return false;
+		CurrentlyLoadedProject->Modules.Empty();
+		for(int32 Idx = 0; Idx < StartupModuleNames->Num(); Idx++)
+		{
+			CurrentlyLoadedProject->Modules.Add(FModuleDescriptor(*(*StartupModuleNames)[Idx]));
+		}
 	}
+
+	// Update file on disk
+	return CurrentlyLoadedProject->Save(FPaths::GetProjectFilePath(), OutFailReason);
 }
 
 bool FProjectManager::SignSampleProject(const FString& FilePath, const FString& Category, FText& OutFailReason)
 {
-	TSharedRef<FProject> NewProject = MakeShareable( new FProject() );
-	if ( !NewProject->LoadFromFile(FilePath, OutFailReason) )
+	FProjectDescriptor Descriptor;
+	if(!Descriptor.Load(FilePath, OutFailReason))
 	{
 		return false;
 	}
 
-	NewProject->SignSampleProject(FilePath, Category);
-
-	const FString& FileContents = NewProject->SerializeToJSON();
-	if (FFileHelper::SaveStringToFile(FileContents, *FilePath))
-	{
-		return true;
-	}
-	else
-	{
-		OutFailReason = FText::Format( LOCTEXT("FailedToSaveSignedProject", "Failed to save signed project file {0}"), FText::FromString(FilePath) );
-		return false;
-	}
+	Descriptor.Sign(FilePath);
+	Descriptor.Category = Category;
+	return Descriptor.Save(FilePath, OutFailReason);
 }
 
 bool FProjectManager::QueryStatusForProject(const FString& FilePath, FProjectStatus& OutProjectStatus) const
 {
-	TSharedRef<FProject> NewProject = MakeShareable( new FProject() );
 	FText FailReason;
-	if ( !NewProject->LoadFromFile(FilePath, FailReason) )
+	FProjectDescriptor Descriptor;
+	if(!Descriptor.Load(FilePath, FailReason))
 	{
 		return false;
 	}
 
-	QueryStatusForProjectImpl(*NewProject, FilePath, OutProjectStatus);
+	QueryStatusForProjectImpl(Descriptor, FilePath, OutProjectStatus);
 	return true;
 }
 
@@ -317,31 +207,37 @@ bool FProjectManager::QueryStatusForCurrentProject(FProjectStatus& OutProjectSta
 	return true;
 }
 
-void FProjectManager::QueryStatusForProjectImpl(const FProject& Project, const FString& FilePath, FProjectStatus& OutProjectStatus)
+void FProjectManager::QueryStatusForProjectImpl(const FProjectDescriptor& ProjectInfo, const FString& FilePath, FProjectStatus& OutProjectStatus)
 {
-	const FProjectInfo& ProjectInfo = Project.GetProjectInfo();
-	OutProjectStatus.Name = ProjectInfo.Name;
+	OutProjectStatus.Name = FPaths::GetBaseFilename(FilePath);
 	OutProjectStatus.Description = ProjectInfo.Description;
 	OutProjectStatus.Category = ProjectInfo.Category;
 	OutProjectStatus.bCodeBasedProject = ProjectInfo.Modules.Num() > 0;
-	OutProjectStatus.bSignedSampleProject = Project.IsSignedSampleProject(FilePath);
-	OutProjectStatus.bRequiresUpdate = Project.RequiresUpdate();
+	OutProjectStatus.bSignedSampleProject = ProjectInfo.IsSigned(FilePath);
+	OutProjectStatus.bRequiresUpdate = ProjectInfo.FileVersion < EProjectDescriptorVersion::Latest;
 	OutProjectStatus.TargetPlatforms = ProjectInfo.TargetPlatforms;
 }
 
 void FProjectManager::UpdateSupportedTargetPlatformsForProject(const FString& FilePath, const FName& InPlatformName, const bool bIsSupported)
 {
-	TSharedRef<FProject> NewProject = MakeShareable( new FProject() );
+	FProjectDescriptor NewProject;
+
 	FText FailReason;
-	if ( !NewProject->LoadFromFile(FilePath, FailReason) )
+	if ( !NewProject.Load(FilePath, FailReason) )
 	{
 		return;
 	}
 
-	NewProject->UpdateSupportedTargetPlatforms(InPlatformName, bIsSupported);
+	if ( bIsSupported )
+	{
+		NewProject.TargetPlatforms.AddUnique(InPlatformName);
+	}
+	else
+	{
+		NewProject.TargetPlatforms.Remove(InPlatformName);
+	}
 
-	const FString& FileContents = NewProject->SerializeToJSON();
-	FFileHelper::SaveStringToFile(FileContents, *FilePath);
+	NewProject.Save(FilePath, FailReason);
 
 	// Call OnTargetPlatformsForCurrentProjectChangedEvent if this project is the same as the one we currently have loaded
 	const FString CurrentProjectPath = FPaths::ConvertRelativePathToFull(FPaths::GetProjectFilePath());
@@ -361,25 +257,24 @@ void FProjectManager::UpdateSupportedTargetPlatformsForCurrentProject(const FNam
 
 	CurrentlyLoadedProject->UpdateSupportedTargetPlatforms(InPlatformName, bIsSupported);
 
-	const FString& FileContents = CurrentlyLoadedProject->SerializeToJSON();
-	FFileHelper::SaveStringToFile(FileContents, *FPaths::GetProjectFilePath());
+	FText FailReason;
+	CurrentlyLoadedProject->Save(FPaths::GetProjectFilePath(), FailReason);
 
 	OnTargetPlatformsForCurrentProjectChangedEvent.Broadcast();
 }
 
 void FProjectManager::ClearSupportedTargetPlatformsForProject(const FString& FilePath)
 {
-	TSharedRef<FProject> NewProject = MakeShareable( new FProject() );
+	FProjectDescriptor Descriptor;
+
 	FText FailReason;
-	if ( !NewProject->LoadFromFile(FilePath, FailReason) )
+	if ( !Descriptor.Load(FilePath, FailReason) )
 	{
 		return;
 	}
 
-	NewProject->ClearSupportedTargetPlatforms();
-
-	const FString& FileContents = NewProject->SerializeToJSON();
-	FFileHelper::SaveStringToFile(FileContents, *FilePath);
+	Descriptor.TargetPlatforms.Empty();
+	Descriptor.Save(FilePath, FailReason);
 
 	// Call OnTargetPlatformsForCurrentProjectChangedEvent if this project is the same as the one we currently have loaded
 	const FString CurrentProjectPath = FPaths::ConvertRelativePathToFull(FPaths::GetProjectFilePath());
@@ -397,10 +292,10 @@ void FProjectManager::ClearSupportedTargetPlatformsForCurrentProject()
 		return;
 	}
 
-	CurrentlyLoadedProject->ClearSupportedTargetPlatforms();
+	CurrentlyLoadedProject->TargetPlatforms.Empty();
 
-	const FString& FileContents = CurrentlyLoadedProject->SerializeToJSON();
-	FFileHelper::SaveStringToFile(FileContents, *FPaths::GetProjectFilePath());
+	FText FailReason;
+	CurrentlyLoadedProject->Save(FPaths::GetProjectFilePath(), FailReason);
 
 	OnTargetPlatformsForCurrentProjectChangedEvent.Broadcast();
 }
diff --git a/Engine/Source/Runtime/Projects/Private/ProjectManager.h b/Engine/Source/Runtime/Projects/Private/ProjectManager.h
--- a/Engine/Source/Runtime/Projects/Private/ProjectManager.h
+++ b/Engine/Source/Runtime/Projects/Private/ProjectManager.h
@@ -2,89 +2,7 @@
 
 #pragma once
 
-
-/**
- * Info about a single project
- */
-class FProjectInfo : public FProjectOrPluginInfo
-{
-public:
-
-	//////////////////////////////////////////////////////////////////////////////////////////////////////
-	// Everything Below this line is serialized. If you add or remove properties, you must update
-	// FProject::PerformAdditionalSerialization and FProject::PerformAdditionalDeserialization
-	//////////////////////////////////////////////////////////////////////////////////////////////////////
-
-	/** A hash that is used to determine if the project was forked from a sample */
-	uint32 EpicSampleNameHash;
-
-	/** Array of platforms that this project is targeting */
-	TArray<FName> TargetPlatforms;
-
-	/** Project info constructor */
-	FProjectInfo()
-		: FProjectOrPluginInfo()
-		, EpicSampleNameHash(0)
-	{
-	}
-};
-
-
-/**
- * Instance of a project in memory
- */
-class FProject : public FProjectOrPlugin
-{
-public:
-	/** Constructor */
-	FProject();
-	FProject( const FProjectInfo& InitProjectInfo );
-
-	/**
-	 * Gets whether the project is a sample project signed with the given filename.
-	 *
-	 * @param FilePath	The filepath from which the project was loaded.
-	 * @return true if the project has been signed with a filename matching the FilePath param.
-	 */
-	bool IsSignedSampleProject(const FString& FilePath) const;
-
-	/**
-	 * Sets the project's EpicSampleNameHash based on a filename, and sets its category.
-	 * This marks the project as a sample and fixes its filename so that
-	 * it isn't mistaken for a sample if a copy of the file is made.
-	 *
-	 * @param FilePath	The filepath where the sample project is stored.
-	 */
-	void SignSampleProject(const FString& FilePath, const FString &Category);
-
-	/**
-	 * Update the list of supported target platforms based upon the parameters provided
-	 * 
-	 * @param	InPlatformName		Name of the platform to target (eg, WindowsNoEditor)
-	 * @param	bIsSupported		true if the platform should be supported by this project, false if it should not
-	 */
-	void UpdateSupportedTargetPlatforms(const FName& InPlatformName, const bool bIsSupported);
-
-	/** Clear the list of supported target platforms */
-	void ClearSupportedTargetPlatforms();
-
-	const FProjectInfo& GetProjectInfo() const
-	{
-		return ProjectInfo;
-	}
-
-protected:
-	virtual FProjectOrPluginInfo& GetProjectOrPluginInfo() override { return ProjectInfo; }
-	virtual const FProjectOrPluginInfo& GetProjectOrPluginInfo() const override { return ProjectInfo; }
-
-	virtual bool PerformAdditionalDeserialization(const TSharedRef< FJsonObject >& FileObject) override;
-	virtual void PerformAdditionalSerialization(const TSharedRef< TJsonWriter<> >& Writer) const override;
-
-private:
-	/** The descriptor of this project that we loaded from disk */
-	FProjectInfo ProjectInfo;
-};
-
+#include "ProjectDescriptor.h"
 
 /**
  * ProjectAndPluginManager manages available code and content extensions (both loaded and not loaded.)
@@ -113,10 +31,10 @@ public:
 	virtual FOnTargetPlatformsForCurrentProjectChangedEvent& OnTargetPlatformsForCurrentProjectChanged() override { return OnTargetPlatformsForCurrentProjectChangedEvent; }
 
 private:
-	static void QueryStatusForProjectImpl(const FProject& Project, const FString& FilePath, FProjectStatus& OutProjectStatus);
+	static void QueryStatusForProjectImpl(const FProjectDescriptor& Project, const FString& FilePath, FProjectStatus& OutProjectStatus);
 
 	/** The project that is currently loaded in the editor */
-	TSharedPtr< FProject > CurrentlyLoadedProject;
+	TSharedPtr< FProjectDescriptor > CurrentlyLoadedProject;
 
 	/** Delegate called when the target platforms for the current project are changed */
 	FOnTargetPlatformsForCurrentProjectChangedEvent OnTargetPlatformsForCurrentProjectChangedEvent;
diff --git a/Engine/Source/Runtime/Projects/Private/ProjectsPrivatePCH.h b/Engine/Source/Runtime/Projects/Private/ProjectsPrivatePCH.h
--- a/Engine/Source/Runtime/Projects/Private/ProjectsPrivatePCH.h
+++ b/Engine/Source/Runtime/Projects/Private/ProjectsPrivatePCH.h
@@ -20,7 +20,6 @@
 /* Private includes
  *****************************************************************************/
 
- #include "PluginManagerShared.h"
  #include "PluginManager.h"
  #include "ProjectManager.h"
  
\ No newline at end of file
diff --git a/Engine/Source/Runtime/Projects/Public/ModuleDescriptor.h b/Engine/Source/Runtime/Projects/Public/ModuleDescriptor.h
--- a/Engine/Source/Runtime/Projects/Public/ModuleDescriptor.h
+++ b/Engine/Source/Runtime/Projects/Public/ModuleDescriptor.h
@@ -4,6 +4,8 @@
 
 #include "Json.h"
 
+enum class EModuleLoadResult;
+
 /**
  * Phase at which this module should be loaded during startup.
  */
@@ -30,14 +32,6 @@ namespace ELoadingPhase
 		Max
 	};
 
-	/**
-	 * Returns the earliest of two load phases.
-	 *
-	 * @param	The value to convert to a string
-	 * @return	The string representation of this enum value
-	 */
-	ELoadingPhase::Type Earliest( ELoadingPhase::Type First, ELoadingPhase::Type Second );
-	
 	/**
 	 * Converts a string to a ELoadingPhase::Type value
 	 *
@@ -111,8 +105,8 @@ struct FModuleDescriptor
 	/** List of disallowed platforms */
 	TArray<FString> BlacklistPlatforms;
 
-	/** Module info constructor */
-	FModuleDescriptor();
+	/** Normal constructor */
+	FModuleDescriptor(const FName InName = NAME_None, EHostType::Type InType = EHostType::Runtime, ELoadingPhase::Type InLoadingPhase = ELoadingPhase::Default);
 
 	/** Reads a descriptor from the given JSON object */
 	bool Read(const FJsonObject& Object, FText& OutFailReason);
@@ -131,4 +125,10 @@ struct FModuleDescriptor
 
 	/** Tests whether the module should be loaded for the current engine configuration */
 	bool IsLoadedInCurrentConfiguration() const;
+
+	/** Loads all the modules for a given loading phase. Returns a map of module names to load errors */
+	static void LoadModulesForPhase(ELoadingPhase::Type LoadingPhase, const TArray<FModuleDescriptor>& Modules, TMap<FName, EModuleLoadResult>& ModuleLoadErrors);
+
+	/** Checks whether all the modules in an array are up to date */
+	static bool AreModulesUpToDate(const TArray<FModuleDescriptor>& Modules);
 };
diff --git a/Engine/Source/Runtime/Projects/Public/PluginDescriptor.h b/Engine/Source/Runtime/Projects/Public/PluginDescriptor.h
new file mode 100644
--- /dev/null
+++ b/Engine/Source/Runtime/Projects/Public/PluginDescriptor.h
@@ -0,0 +1,75 @@
+// Copyright 1998-2014 Epic Games, Inc. All Rights Reserved.
+
+#pragma once
+
+#include "ModuleDescriptor.h"
+
+/**
+ * Version numbers for plugin descriptors.
+ */ 
+namespace EPluginDescriptorVersion
+{
+	enum Type
+	{
+		Invalid = 0,
+		Initial = 1,
+		NameHash = 2,
+		ProjectPluginUnification = 3,
+		// !!!!!!!!!! IMPORTANT: Remember to also update LatestPluginDescriptorFileVersion in Plugins.cs (and Plugin system documentation) when this changes!!!!!!!!!!!
+		// -----<new versions can be added before this line>-------------------------------------------------
+		// - this needs to be the last line (see note below)
+		LatestPlusOne,
+		Latest = LatestPlusOne - 1
+	};
+}
+
+/**
+ * Descriptor for plugins. Contains all the information contained within a .uplugin file.
+ */
+struct PROJECTS_API FPluginDescriptor
+{
+	/** Descriptor version number */
+	int32 FileVersion;
+
+	/** Version number for the plugin.  The version number must increase with every version of the plugin, so that the system 
+	    can determine whether one version of a plugin is newer than another, or to enforce other requirements.  This version
+		number is not displayed in front-facing UI.  Use the VersionName for that. */
+	int32 Version;
+
+	/** Name of the version for this plugin.  This is the front-facing part of the version number.  It doesn't need to match
+	    the version number numerically, but should be updated when the version number is increased accordingly. */
+	FString VersionName;
+
+	/** Friendly name of the plugin */
+	FString FriendlyName;
+
+	/** Description of the plugin */
+	FString Description;
+
+	/** The name of the category this plugin */
+	FString Category;
+
+	/** The company or individual who created this plugin.  This is an optional field that may be displayed in the user interface. */
+	FString CreatedBy;
+
+	/** Hyperlink URL string for the company or individual who created this plugin.  This is optional. */
+	FString CreatedByURL;
+
+	/** List of all modules associated with this plugin */
+	TArray<FModuleDescriptor> Modules;
+
+	/** Can this plugin contain content? */
+	bool bCanContainContent;
+
+	/** Marks the plugin as beta in the UI */
+	bool bIsBetaVersion;
+
+	/** Constructor. */
+	FPluginDescriptor();
+
+	/** Loads the descriptor from the given file. */
+	bool Load(const FString& FileName, FText& OutFailReason);
+
+	/** Reads the descriptor from the given JSON object */
+	bool Read(const FJsonObject& Object, FText& OutFailReason);
+};
diff --git a/Engine/Source/Runtime/Projects/Public/ProjectDescriptor.h b/Engine/Source/Runtime/Projects/Public/ProjectDescriptor.h
new file mode 100644
--- /dev/null
+++ b/Engine/Source/Runtime/Projects/Public/ProjectDescriptor.h
@@ -0,0 +1,75 @@
+// Copyright 1998-2014 Epic Games, Inc. All Rights Reserved.
+
+#pragma once
+
+#include "Projects.h"
+
+/**
+ * Version numbers for project descriptors.
+ */ 
+namespace EProjectDescriptorVersion
+{
+	enum Type
+	{
+		Invalid = 0,
+		Initial = 1,
+		NameHash = 2,
+		ProjectPluginUnification = 3,
+		// !!!!!!!!!! IMPORTANT: Remember to also update LatestPluginDescriptorFileVersion in Plugins.cs (and Plugin system documentation) when this changes!!!!!!!!!!!
+		// -----<new versions can be added before this line>-------------------------------------------------
+		// - this needs to be the last line (see note below)
+		LatestPlusOne,
+		Latest = LatestPlusOne - 1
+	};
+}
+
+/**
+ * Descriptor for projects. Contains all the information contained within a .uproject file.
+ */
+struct PROJECTS_API FProjectDescriptor
+{
+	/** Descriptor version number. */
+	EProjectDescriptorVersion::Type FileVersion;
+
+	/** The associated engine label. Used by UnrealVersionSelector to open projects with the correct editor */
+	FString EngineAssociation;
+
+	/** Category to show under the project browser */
+	FString Category;
+
+	/** Description to show in the project browser */
+	FString Description;
+
+	/** List of all modules associated with this plugin */
+	TArray<FModuleDescriptor> Modules;
+
+	/** Array of platforms that this project is targeting */
+	TArray<FName> TargetPlatforms;
+
+	/** A hash that is used to determine if the project was forked from a sample */
+	uint32 EpicSampleNameHash;
+
+	/** Constructor. */
+	FProjectDescriptor();
+
+	/** Signs the project given for the given filename */
+	void Sign(const FString& FilePath);
+
+	/** Checks whether the descriptor is signed */
+	bool IsSigned(const FString& FilePath) const;
+
+	/** Updates the supported target platforms list */
+	void UpdateSupportedTargetPlatforms(const FName& InPlatformName, bool bIsSupported);
+
+	/** Loads the descriptor from the given file. */
+	bool Load(const FString& FileName, FText& OutFailReason);
+
+	/** Reads the descriptor from the given JSON object */
+	bool Read(const FJsonObject& Object, FText& OutFailReason);
+
+	/** Saves the descriptor to the given file. */
+	bool Save(const FString& FileName, FText& OutFailReason);
+
+	/** Writes the descriptor to the given JSON object */
+	void Write(TJsonWriter<>& Writer) const;
+};
