Store the list of enabled plugins in the uproject file (fixing TTP 335060), and enable all game plugins by default (TTP 339313).

diff --git a/Engine/Plugins/Editor/PluginsEditor/Source/PluginsEditor/PluginsEditor.Build.cs b/Engine/Plugins/Editor/PluginsEditor/Source/PluginsEditor/PluginsEditor.Build.cs
--- a/Engine/Plugins/Editor/PluginsEditor/Source/PluginsEditor/PluginsEditor.Build.cs
+++ b/Engine/Plugins/Editor/PluginsEditor/Source/PluginsEditor/PluginsEditor.Build.cs
@@ -22,6 +22,7 @@ namespace UnrealBuildTool.Rules
 					"EditorStyle",
 					"Projects",
 					"UnrealEd",
+					"GameProjectGeneration",
 				}
 			);
 		}
diff --git a/Engine/Plugins/Editor/PluginsEditor/Source/PluginsEditor/Private/SPluginListTile.cpp b/Engine/Plugins/Editor/PluginsEditor/Source/PluginsEditor/Private/SPluginListTile.cpp
--- a/Engine/Plugins/Editor/PluginsEditor/Source/PluginsEditor/Private/SPluginListTile.cpp
+++ b/Engine/Plugins/Editor/PluginsEditor/Source/PluginsEditor/Private/SPluginListTile.cpp
@@ -6,6 +6,7 @@
 #include "SPluginsEditor.h"
 #include "SPluginList.h"
 #include "PluginStyle.h"
+#include "GameProjectGenerationModule.h"
 
 #define LOCTEXT_NAMESPACE "PluginListTile"
 
@@ -275,7 +276,14 @@ void SPluginListTile::OnEnablePluginCheckboxChanged(ESlateCheckBoxState::Type Ne
 	}
 
 	ItemData->PluginStatus.bIsEnabled = bNewEnabledState;
-	IPluginManager::Get().SetPluginEnabled(ItemData->PluginStatus.Name, bNewEnabledState);
+
+	FGameProjectGenerationModule::Get().TryMakeProjectFileWriteable();
+
+	FText FailMessage;
+	if(!IProjectManager::Get().SetPluginEnabled(ItemData->PluginStatus.Name, bNewEnabledState, FailMessage))
+	{
+		FMessageDialog::Open(EAppMsgType::Ok, FailMessage);
+	}
 }
 
 
diff --git a/Engine/Plugins/Editor/PluginsEditor/Source/PluginsEditor/Private/SPluginsEditor.cpp b/Engine/Plugins/Editor/PluginsEditor/Source/PluginsEditor/Private/SPluginsEditor.cpp
--- a/Engine/Plugins/Editor/PluginsEditor/Source/PluginsEditor/Private/SPluginsEditor.cpp
+++ b/Engine/Plugins/Editor/PluginsEditor/Source/PluginsEditor/Private/SPluginsEditor.cpp
@@ -124,7 +124,7 @@ void SPluginsEditor::Construct( const FArguments& Args )
 
 EVisibility SPluginsEditor::HandleRestartEditorNoticeVisibility() const
 {
-	return IPluginManager::Get().IsRestartRequired() ? EVisibility::Visible : EVisibility::Collapsed;
+	return IProjectManager::Get().IsRestartRequired() ? EVisibility::Visible : EVisibility::Collapsed;
 }
 
 
diff --git a/Engine/Source/Editor/GameProjectGeneration/Private/GameProjectGenerationModule.cpp b/Engine/Source/Editor/GameProjectGeneration/Private/GameProjectGenerationModule.cpp
--- a/Engine/Source/Editor/GameProjectGeneration/Private/GameProjectGenerationModule.cpp
+++ b/Engine/Source/Editor/GameProjectGeneration/Private/GameProjectGenerationModule.cpp
@@ -43,6 +43,10 @@ void FGameProjectGenerationModule::OpenAddCodeToProjectDialog()
 	AddCodeToProjectDialogOpenedEvent.Broadcast();
 }
 
+void FGameProjectGenerationModule::TryMakeProjectFileWriteable()
+{
+	GameProjectUtils::TryMakeProjectFileWriteable();
+}
 
 void FGameProjectGenerationModule::CheckForOutOfDateGameProjectFile()
 {
diff --git a/Engine/Source/Editor/GameProjectGeneration/Private/GameProjectUtils.cpp b/Engine/Source/Editor/GameProjectGeneration/Private/GameProjectUtils.cpp
--- a/Engine/Source/Editor/GameProjectGeneration/Private/GameProjectUtils.cpp
+++ b/Engine/Source/Editor/GameProjectGeneration/Private/GameProjectUtils.cpp
@@ -2464,32 +2464,36 @@ void GameProjectUtils::OnUpdateProjectCancel()
 	}
 }
 
-bool GameProjectUtils::UpdateGameProjectFile(const FString& ProjectFilename, const FString& EngineIdentifier, const TArray<FString>* StartupModuleNames, bool& OutbWasCheckedOut, FText& OutFailReason)
+void GameProjectUtils::TryMakeProjectFileWriteable()
 {
+	FString ProjectFileName = FPaths::GetProjectFilePath();
+
 	// First attempt to check out the file if SCC is enabled
 	if ( ISourceControlModule::Get().IsEnabled() )
 	{
-		OutbWasCheckedOut = GameProjectUtils::CheckoutGameProjectFile(ProjectFilename, OutFailReason);
-		if ( !OutbWasCheckedOut )
-		{
-			// Failed to check out the file
-			return false;
-		}
+		FText FailReason;
+		GameProjectUtils::CheckoutGameProjectFile(ProjectFileName, FailReason);
 	}
-	else
+
+	// Check if it's writeable
+	if(FPlatformFileManager::Get().GetPlatformFile().IsReadOnly(*ProjectFileName))
 	{
-		if(FPlatformFileManager::Get().GetPlatformFile().IsReadOnly(*ProjectFilename))
+		FText ShouldMakeProjectWriteable = LOCTEXT("ShouldMakeProjectWriteable_Message", "'{ProjectFilename}' is read-only and cannot be updated, would you like to make it writeable?");
+
+		FFormatNamedArguments Arguments;
+		Arguments.Add( TEXT("ProjectFilename"), FText::FromString(ProjectFileName));
+
+		if(FMessageDialog::Open(EAppMsgType::YesNo, FText::Format(ShouldMakeProjectWriteable, Arguments)) == EAppReturnType::Yes)
 		{
-			FText ShouldMakeProjectWriteable = LOCTEXT("ShouldMakeProjectWriteable_Message", "'{ProjectFilename}' is read-only and cannot be updated, would you like to make it writeable?");
-			FFormatNamedArguments Arguments;
-			Arguments.Add( TEXT("ProjectFilename"), FText::FromString(ProjectFilename));
-			if ( FMessageDialog::Open(EAppMsgType::YesNo, FText::Format(ShouldMakeProjectWriteable, Arguments)) == EAppReturnType::Yes )
-			{
-				FPlatformFileManager::Get().GetPlatformFile().SetReadOnly(*ProjectFilename, false);
-			}
+			FPlatformFileManager::Get().GetPlatformFile().SetReadOnly(*ProjectFileName, false);
 		}
-		OutbWasCheckedOut = false;
 	}
+}
+
+bool GameProjectUtils::UpdateGameProjectFile(const FString& ProjectFilename, const FString& EngineIdentifier, const TArray<FString>* StartupModuleNames, bool& OutbWasCheckedOut, FText& OutFailReason)
+{
+	// Make sure we can write to the project file
+	TryMakeProjectFileWriteable();
 
 	// Now tell the project manager to update the file
 	if (!IProjectManager::Get().UpdateLoadedProjectFileToCurrent(StartupModuleNames, EngineIdentifier, OutFailReason))
diff --git a/Engine/Source/Editor/GameProjectGeneration/Private/GameProjectUtils.h b/Engine/Source/Editor/GameProjectGeneration/Private/GameProjectUtils.h
--- a/Engine/Source/Editor/GameProjectGeneration/Private/GameProjectUtils.h
+++ b/Engine/Source/Editor/GameProjectGeneration/Private/GameProjectUtils.h
@@ -141,6 +141,9 @@ public:
 	/** Warn the user if the project filename is invalid in case they renamed it outside the editor */
 	static void CheckAndWarnProjectFilenameValid();
 
+	/** Checks out the current project file (or prompts to make writable) */
+	static void TryMakeProjectFileWriteable();
+
 	/** Updates the currently loaded project. Returns true if the project was updated successfully or if no update was needed */
 	static bool UpdateGameProject(const FString &EngineIdentifier);
 
diff --git a/Engine/Source/Editor/GameProjectGeneration/Public/GameProjectGenerationModule.h b/Engine/Source/Editor/GameProjectGeneration/Public/GameProjectGenerationModule.h
--- a/Engine/Source/Editor/GameProjectGeneration/Public/GameProjectGenerationModule.h
+++ b/Engine/Source/Editor/GameProjectGeneration/Public/GameProjectGenerationModule.h
@@ -47,6 +47,9 @@ public:
 	DECLARE_EVENT(FGameProjectGenerationModule, FAddCodeToProjectDialogOpenedEvent);
 	FAddCodeToProjectDialogOpenedEvent& OnAddCodeToProjectDialogOpened() { return AddCodeToProjectDialogOpenedEvent; }
 
+	/** Tries to make the project file writable. Prompts to check out as necessary. */
+	virtual void TryMakeProjectFileWriteable();
+
 	/** Prompts the user to update his project file, if necessary. */
 	virtual void CheckForOutOfDateGameProjectFile();
 
diff --git a/Engine/Source/Editor/MainFrame/Private/Frame/MainFrameActions.cpp b/Engine/Source/Editor/MainFrame/Private/Frame/MainFrameActions.cpp
--- a/Engine/Source/Editor/MainFrame/Private/Frame/MainFrameActions.cpp
+++ b/Engine/Source/Editor/MainFrame/Private/Frame/MainFrameActions.cpp
@@ -508,7 +508,7 @@ void FMainFrameActionCallbacks::PackageProject( const FName InPlatformInfoName )
 		FMessageDialog::Open(EAppMsgType::Ok, LOCTEXT("PackagingNotWorkingMessage", "Sorry, packaging is currently not supported for code-based iOS projects. This feature will be available in a future release.") );
 		return;
 	}
-	if (FRocketSupport::IsRocket() && PlatformInfo->TargetPlatformName == FName("IOS") && IPluginManager::Get().HasThirdPartyPlugin())
+	if (FRocketSupport::IsRocket() && PlatformInfo->TargetPlatformName == FName("IOS") && IProjectManager::Get().IsThirdPartyPluginEnabled())
 	{
 		FMessageDialog::Open(EAppMsgType::Ok, LOCTEXT("PackagingNotWorkingMessage", "Sorry, packaging is currently not supported for content based projects with third-party plugins. This feature will be available in a future release.") );
 		return;
diff --git a/Engine/Source/Editor/UnrealEd/Private/PlayLevel.cpp b/Engine/Source/Editor/UnrealEd/Private/PlayLevel.cpp
--- a/Engine/Source/Editor/UnrealEd/Private/PlayLevel.cpp
+++ b/Engine/Source/Editor/UnrealEd/Private/PlayLevel.cpp
@@ -1291,7 +1291,7 @@ void UEditorEngine::PlayUsingLauncher()
 			bPlayUsingLauncher = false;
 			return;
 		}
-		if (FRocketSupport::IsRocket() && PlayUsingLauncherDeviceId.Left(PlayUsingLauncherDeviceId.Find(TEXT("@"))) == TEXT("IOS") && IPluginManager::Get().HasThirdPartyPlugin())
+		if (FRocketSupport::IsRocket() && PlayUsingLauncherDeviceId.Left(PlayUsingLauncherDeviceId.Find(TEXT("@"))) == TEXT("IOS") && IProjectManager::Get().IsThirdPartyPluginEnabled())
 		{
 			FMessageDialog::Open(EAppMsgType::Ok, LOCTEXT("LaunchOnDeviceNotWorkingMessage", "Sorry, launching on a device is currently not supported for content based projects with third-party plugins. This feature will be available in a future release.") );
 			bPlayUsingLauncher = false;
diff --git a/Engine/Source/Programs/UnrealBuildTool/Configuration/UEBuildTarget.cs b/Engine/Source/Programs/UnrealBuildTool/Configuration/UEBuildTarget.cs
--- a/Engine/Source/Programs/UnrealBuildTool/Configuration/UEBuildTarget.cs
+++ b/Engine/Source/Programs/UnrealBuildTool/Configuration/UEBuildTarget.cs
@@ -2145,7 +2145,18 @@ namespace UnrealBuildTool
 							FilterPluginNames.AddRange(ConfigPlugins);
 						}
 					}
-					EnabledPlugins.AddRange(ValidPlugins.Where(x => FilterPluginNames.Contains(x.Name)));
+
+					// Update the plugin list for game targets
+					if(Rules.Type != TargetRules.TargetType.Program && UnrealBuildTool.HasUProjectFile())
+					{
+						// Enable all the game specific plugins by default
+						FilterPluginNames.AddRange(ValidPlugins.Where(x => x.LoadedFrom == PluginInfo.LoadedFromType.GameProject).Select(x => x.Name));
+
+						// Use the project settings to update the plugin list for this target
+						FilterPluginNames = UProjectInfo.GetEnabledPlugins(UnrealBuildTool.GetUProjectFile(), FilterPluginNames, Platform);
+					}
+
+					EnabledPlugins.AddRange(ValidPlugins.Where(x => FilterPluginNames.Contains(x.Name)).Distinct());
 				}
 				else
 				{
diff --git a/Engine/Source/Programs/UnrealBuildTool/System/Plugins.cs b/Engine/Source/Programs/UnrealBuildTool/System/Plugins.cs
--- a/Engine/Source/Programs/UnrealBuildTool/System/Plugins.cs
+++ b/Engine/Source/Programs/UnrealBuildTool/System/Plugins.cs
@@ -69,6 +69,64 @@ namespace UnrealBuildTool
 		}
 	}
 
+	class PluginReferenceDescriptor
+	{
+		public string Name;
+		public bool bEnabled;
+		public List<UnrealTargetPlatform> WhitelistPlatforms = new List<UnrealTargetPlatform>();
+		public List<UnrealTargetPlatform> BlacklistPlatforms = new List<UnrealTargetPlatform>();
+
+		public PluginReferenceDescriptor(string InName, bool bInEnabled)
+		{
+			Name = InName;
+			bEnabled = bInEnabled;
+		}
+
+		public static PluginReferenceDescriptor FromJson(Dictionary<string, object> Dictionary)
+		{
+			PluginReferenceDescriptor Descriptor = new PluginReferenceDescriptor((string)Dictionary["Name"], (bool)Dictionary["Enabled"]);
+			TryParsePlatformList(Dictionary, "WhitelistPlatforms", Descriptor.WhitelistPlatforms);
+			TryParsePlatformList(Dictionary, "BlacklistPlatforms", Descriptor.BlacklistPlatforms);
+			return Descriptor;
+		}
+
+		public bool IsEnabledForPlatform(UnrealTargetPlatform Platform)
+		{
+			if(!bEnabled)
+			{
+				return false;
+			}
+			if(WhitelistPlatforms.Count > 0 && !WhitelistPlatforms.Contains(Platform))
+			{
+				return false;
+			}
+			if(BlacklistPlatforms.Contains(Platform))
+			{
+				return false;
+			}
+			return true;
+		}
+
+		static void TryParsePlatformList(Dictionary<string, object> Dictionary, string FieldName, List<UnrealTargetPlatform> Platforms)
+		{
+			object ArrayObject;
+			if(Dictionary.TryGetValue(FieldName, out ArrayObject))
+			{
+				foreach(string PlatformName in ((object[])ArrayObject).Select(x => (string)x))
+				{
+					UnrealTargetPlatform Platform;
+					if(Enum.TryParse(PlatformName, true, out Platform))
+					{
+						Platforms.Add(Platform);
+					}
+					else
+					{
+						throw new BuildException("Unknown platform name '{0}'", PlatformName);
+					}
+				}
+			}
+		}
+	}
 
 	public class Plugins
 	{
diff --git a/Engine/Source/Programs/UnrealBuildTool/System/UProjectInfo.cs b/Engine/Source/Programs/UnrealBuildTool/System/UProjectInfo.cs
--- a/Engine/Source/Programs/UnrealBuildTool/System/UProjectInfo.cs
+++ b/Engine/Source/Programs/UnrealBuildTool/System/UProjectInfo.cs
@@ -292,5 +292,48 @@ namespace UnrealBuildTool
 			}
 			return "";
 		}
+
+		/// <summary>
+		/// Get the project folder for the given project name
+		/// </summary>
+		/// <param name="InProjectName">Name of the project of interest</param>
+		/// <returns>The project filename, empty string if not found</returns>
+		public static List<string> GetEnabledPlugins(string ProjectFileName, List<string> DefaultEnabledPlugins, UnrealTargetPlatform Platform)
+		{
+			// Parse the project file
+			string ProjectFile = File.ReadAllText(ProjectFileName);
+			try
+			{
+				List<string> EnabledPlugins = new List<string>(DefaultEnabledPlugins);
+
+				// Create a case-insensitive dictionary of the contents
+				Dictionary<string, object> Descriptor = fastJSON.JSON.Instance.ToObject<Dictionary<string, object>>(ProjectFile);
+				Descriptor = new Dictionary<string,object>(Descriptor, StringComparer.InvariantCultureIgnoreCase);
+
+				// Get the list of plugins
+				object PluginsObject;
+				if(Descriptor.TryGetValue("Plugins", out PluginsObject))
+				{
+					object[] PluginsArrayObject = (object[])PluginsObject;
+					foreach(object PluginObject in PluginsArrayObject)
+					{
+						PluginReferenceDescriptor Plugin = PluginReferenceDescriptor.FromJson((Dictionary<string, object>)PluginObject);
+						if(Plugin.IsEnabledForPlatform(Platform))
+						{
+							EnabledPlugins.Add(Plugin.Name);
+						}
+						else
+						{
+							EnabledPlugins.RemoveAll(x => x.Equals(Plugin.Name, StringComparison.CurrentCultureIgnoreCase));
+						}
+					}
+				}
+				return EnabledPlugins;
+			}
+			catch(BuildException Exception)
+			{
+				throw new BuildException("While parsing '{0}': {1}", ProjectFileName, Exception.ToString());
+			}
+		}
 	}
 }
diff --git a/Engine/Source/Runtime/Projects/Private/PluginDescriptor.cpp b/Engine/Source/Runtime/Projects/Private/PluginDescriptor.cpp
--- a/Engine/Source/Runtime/Projects/Private/PluginDescriptor.cpp
+++ b/Engine/Source/Runtime/Projects/Private/PluginDescriptor.cpp
@@ -91,4 +91,123 @@ bool FPluginDescriptor::Read(const FJsonObject& Object, FText& OutFailReason)
 	return true;
 }
 
+
+
+
+
+
+
+
+FPluginReferenceDescriptor::FPluginReferenceDescriptor(const FString &InName, bool bInEnabled)
+	: Name(InName)
+	, bEnabled(bInEnabled)
+{
+}
+
+bool FPluginReferenceDescriptor::IsEnabledForPlatform(const FString& Platform) const
+{
+	// If it's not enabled at all, return false
+	if(!bEnabled)
+	{
+		return false;
+	}
+
+	// If there is a list of whitelisted platforms, and this isn't one of them, return false
+	if(WhitelistPlatforms.Num() > 0 && !WhitelistPlatforms.Contains(Platform))
+	{
+		return false;
+	}
+
+	// If this platform is blacklisted, also return false
+	if(BlacklistPlatforms.Contains(Platform))
+	{
+		return false;
+	}
+
+	return true;
+}
+
+bool FPluginReferenceDescriptor::Read(const FJsonObject& Object, FText& OutFailReason)
+{
+	// Get the name
+	if(!Object.TryGetStringField(TEXT("Name"), Name))
+	{
+		OutFailReason = LOCTEXT("PluginReferenceWithoutName", "Plugin references must have a 'Name' field");
+		return false;
+	}
+
+	// Get the enabled field
+	if(!Object.TryGetBoolField(TEXT("Enabled"), bEnabled))
+	{
+		OutFailReason = LOCTEXT("PluginReferenceWithoutEnabled", "Plugin references must have an 'Enabled' field");
+		return false;
+	}
+
+	// Get the platform lists
+	Object.TryGetStringArrayField(TEXT("WhitelistPlatforms"), WhitelistPlatforms);
+	Object.TryGetStringArrayField(TEXT("BlacklistPlatforms"), BlacklistPlatforms);
+	return true;
+}
+
+bool FPluginReferenceDescriptor::ReadArray(const FJsonObject& Object, const TCHAR* Name, TArray<FPluginReferenceDescriptor>& OutPlugins, FText& OutFailReason)
+{
+	const TArray< TSharedPtr<FJsonValue> > *Array;
+	if(Object.TryGetArrayField(Name, Array))
+	{
+		for(const TSharedPtr<FJsonValue> &Item : *Array)
+		{
+			const TSharedPtr<FJsonObject> *Object;
+			if(Item.IsValid() && Item->TryGetObject(Object))
+			{
+				FPluginReferenceDescriptor Plugin;
+				if(!Plugin.Read(*Object->Get(), OutFailReason))
+				{
+					return false;
+				}
+				OutPlugins.Add(Plugin);
+			}
+		}
+	}
+	return true;
+}
+
+void FPluginReferenceDescriptor::Write(TJsonWriter<>& Writer) const
+{
+	Writer.WriteObjectStart();
+	Writer.WriteValue(TEXT("Name"), Name);
+	Writer.WriteValue(TEXT("Enabled"), bEnabled);
+	if (WhitelistPlatforms.Num() > 0)
+	{
+		Writer.WriteArrayStart(TEXT("WhitelistPlatforms"));
+		for(int Idx = 0; Idx < WhitelistPlatforms.Num(); Idx++)
+		{
+			Writer.WriteValue(WhitelistPlatforms[Idx]);
+		}
+		Writer.WriteArrayEnd();
+	}
+	if (BlacklistPlatforms.Num() > 0)
+	{
+		Writer.WriteArrayStart(TEXT("BlacklistPlatforms"));
+		for(int Idx = 0; Idx < BlacklistPlatforms.Num(); Idx++)
+		{
+			Writer.WriteValue(BlacklistPlatforms[Idx]);
+		}
+		Writer.WriteArrayEnd();
+	}
+	Writer.WriteObjectEnd();
+}
+
+void FPluginReferenceDescriptor::WriteArray(TJsonWriter<>& Writer, const TCHAR* Name, const TArray<FPluginReferenceDescriptor>& Plugins)
+{
+	if(Plugins.Num() > 0)
+	{
+		Writer.WriteArrayStart(Name);
+		for(int Idx = 0; Idx < Plugins.Num(); Idx++)
+		{
+			Plugins[Idx].Write(Writer);
+		}
+		Writer.WriteArrayEnd();
+	}
+}
+
 #undef LOCTEXT_NAMESPACE
diff --git a/Engine/Source/Runtime/Projects/Private/PluginManager.cpp b/Engine/Source/Runtime/Projects/Private/PluginManager.cpp
--- a/Engine/Source/Runtime/Projects/Private/PluginManager.cpp
+++ b/Engine/Source/Runtime/Projects/Private/PluginManager.cpp
@@ -9,10 +9,6 @@ DEFINE_LOG_CATEGORY_STATIC( LogPluginManager, Log, All );
 
 namespace PluginSystemDefs
 {
-	/** Latest supported version for plugin descriptor files.  We can still usually load older versions, but not newer version.
-	    NOTE: This constant exists in UnrealBuildTool code as well. */
-	static const int32 LatestPluginDescriptorFileVersion = 1;	// IMPORTANT: Remember to also update Plugins.cs (and Plugin system documentation) when this changes!
-		
 	/** File extension of plugin descriptor files.
 	    NOTE: This constant exists in UnrealBuildTool code as well. */
 	static const TCHAR PluginDescriptorFileExtension[] = TEXT( ".uplugin" );
@@ -56,7 +52,6 @@ void FPluginInstance::RegisterPluginMountPoints( const IPluginManager::FRegister
 
 FPluginManager::FPluginManager()
 	: bEarliestPhaseProcessed(false)
-	, bRestartRequired(false)
 {
 	DiscoverAllPlugins();
 }
@@ -287,7 +282,11 @@ void FPluginManager::DiscoverAllPlugins()
 void FPluginManager::EnablePluginsThatAreConfiguredToBeEnabled()
 {
 	TArray< FString > EnabledPluginNames;
-	GetPluginsConfiguredToBeEnabled(EnabledPluginNames);
+#if IS_PROGRAM
+	GConfig->GetArray(TEXT("Plugins"), TEXT("ProgramEnabledPlugins"), EnabledPluginNames, GEngineIni);
+#else
+	FProjectManager::Get().GetEnabledPlugins(EnabledPluginNames);
+#endif
 
 	TSet< FString > AllEnabledPlugins;
 	AllEnabledPlugins.Append( MoveTemp(EnabledPluginNames) );
@@ -406,93 +405,6 @@ bool FPluginManager::AreEnabledPluginModulesUpToDate()
 	return true;
 }
 
-void FPluginManager::GetPluginsConfiguredToBeEnabled(TArray<FString>& OutPluginNames) const
-{
-#if !IS_PROGRAM
-	GConfig->GetArray( TEXT("Plugins"), TEXT("EnabledPlugins"), OutPluginNames, GEngineIni );
-#else
-	GConfig->GetArray(TEXT("Plugins"), TEXT("ProgramEnabledPlugins"), OutPluginNames, GEngineIni);
-#endif
-}
-
-bool FPluginManager::HasThirdPartyPlugin() const
-{
-	static bool bInit = false;
-	TArray<FString> StandardPlugins;
-	if (!bInit)
-	{
-		StandardPlugins.Add(TEXT("BlankPlugin"));
-		StandardPlugins.Add(TEXT("UObjectPlugin"));
-		StandardPlugins.Add(TEXT("PluginsEditor"));
-		StandardPlugins.Add(TEXT("EpicSurvey"));
-		StandardPlugins.Add(TEXT("OculusRift"));
-		StandardPlugins.Add(TEXT("CustomMeshComponent"));
-		StandardPlugins.Add(TEXT("MessagingDebugger"));
-		StandardPlugins.Add(TEXT("PerforceSourceControl"));
-		StandardPlugins.Add(TEXT("SubversionSourceControl"));
-		StandardPlugins.Add(TEXT("UdpMessaging"));
-		StandardPlugins.Add(TEXT("WindowsMoviePlayer"));
-		StandardPlugins.Add(TEXT("CableComponent"));
-		StandardPlugins.Add(TEXT("ExampleDeviceProfileSelector"));
-		StandardPlugins.Add(TEXT("WinDualShock"));
-		StandardPlugins.Add(TEXT("VisualStudioSourceCodeAccess"));
-		StandardPlugins.Add(TEXT("XCodeSourceCodeAccess"));
-		StandardPlugins.Add(TEXT("SlateRemote"));
-		StandardPlugins.Add(TEXT("ScriptPlugin"));
-		StandardPlugins.Add(TEXT("ScriptEditorPlugin"));
-		StandardPlugins.Add(TEXT("SpeedTreeImporter"));
-		StandardPlugins.Add(TEXT("AppleMoviePlayer"));
-		StandardPlugins.Add(TEXT("IOSDeviceProfileSelector"));
-		bInit = true;
-	}
-
-	TArray<FString> EnabledPlugins;
-	GConfig->GetArray( TEXT("Plugins"), TEXT("EnabledPlugins"), EnabledPlugins, GEngineIni );
-
-	for (int Index = 0; Index < EnabledPlugins.Num(); ++Index)
-	{
-		if (!StandardPlugins.Contains(EnabledPlugins[Index]))
-		{
-			return true;
-		}
-	}
-	return false;
-}
-
-void FPluginManager::ConfigurePluginToBeEnabled(const FString& PluginName, bool bEnabled)
-{
-	TArray< FString > EnabledPluginNames;
-	GetPluginsConfiguredToBeEnabled(EnabledPluginNames);
-
-	bool bSaveEnabledPluginNames = false;
-	if ( bEnabled )
-	{
-		if ( !EnabledPluginNames.Contains(PluginName) )
-		{
-			EnabledPluginNames.Add(PluginName);
-			bSaveEnabledPluginNames = true;
-		}
-	}
-	else
-	{
-		int32 ExistingEnabledIdx = EnabledPluginNames.Find(PluginName);
-		if ( ExistingEnabledIdx != INDEX_NONE )
-		{
-			EnabledPluginNames.RemoveAt(ExistingEnabledIdx);
-			bSaveEnabledPluginNames = true;
-		}
-	}
-
-	bRestartRequired = true;
-
-	if ( bSaveEnabledPluginNames )
-	{
-		GConfig->SetArray( TEXT("Plugins"), TEXT("EnabledPlugins"), EnabledPluginNames, GEngineIni );
-		GConfig->Flush(false);
-	}
-}
-
-
 IPluginManager& IPluginManager::Get()
 {
 	// Single instance of manager, allocated on demand and destroyed on program exit.
@@ -539,24 +451,4 @@ TArray< FPluginStatus > FPluginManager::QueryStatusForAllPlugins() const
 	return PluginStatuses;
 }
 
-void FPluginManager::SetPluginEnabled( const FString& PluginName, bool bEnabled )
-{
-	for( auto PluginIt( AllPlugins.CreateConstIterator() ); PluginIt; ++PluginIt )
-	{
-		const TSharedRef< FPluginInstance >& Plugin = *PluginIt;
-
-		if ( Plugin->Name == PluginName )
-		{
-			Plugin->bEnabled = bEnabled;
-			ConfigurePluginToBeEnabled(PluginName, bEnabled);
-			break;
-		}
-	}
-}
-
-bool FPluginManager::IsRestartRequired() const
-{
-	return bRestartRequired;
-}
-
 #undef LOCTEXT_NAMESPACE
\ No newline at end of file
diff --git a/Engine/Source/Runtime/Projects/Private/PluginManager.h b/Engine/Source/Runtime/Projects/Private/PluginManager.h
--- a/Engine/Source/Runtime/Projects/Private/PluginManager.h
+++ b/Engine/Source/Runtime/Projects/Private/PluginManager.h
@@ -71,9 +71,6 @@ public:
 	virtual bool IsPluginModule( const FName ModuleName ) const override;
 	virtual bool AreEnabledPluginModulesUpToDate() override;
 	virtual TArray< FPluginStatus > QueryStatusForAllPlugins() const override;
-	virtual void SetPluginEnabled( const FString& PluginName, bool bEnabled ) override;
-	virtual bool IsRestartRequired() const override;
-	virtual bool HasThirdPartyPlugin() const override;
 
 private:
 
@@ -87,12 +84,6 @@ private:
 	/** Registers content mount points for all enabled plugins */
 	void RegisterEnabledPluginMountPoints();
 
-	/** Returns all plugins currently configured to be enabled */
-	void GetPluginsConfiguredToBeEnabled(TArray<FString>& OutPluginNames) const;
-
-	/** Sets the list of currently configured plugins */
-	void ConfigurePluginToBeEnabled(const FString& PluginName, bool bEnabled);
-
 private:
 	/** All of the plugins that we know about */
 	TArray< TSharedRef< FPluginInstance > > AllPlugins;
@@ -106,9 +97,6 @@ private:
 
 	/** The earliest phase that modules are loaded requires some additional initialization */
 	bool bEarliestPhaseProcessed;
-
-	/** true if plugin changes have been made that will only take effect after a restart */
-	bool bRestartRequired;
 };
 
 
diff --git a/Engine/Source/Runtime/Projects/Private/ProjectDescriptor.cpp b/Engine/Source/Runtime/Projects/Private/ProjectDescriptor.cpp
--- a/Engine/Source/Runtime/Projects/Private/ProjectDescriptor.cpp
+++ b/Engine/Source/Runtime/Projects/Private/ProjectDescriptor.cpp
@@ -90,6 +90,12 @@ bool FProjectDescriptor::Read(const FJsonObject& Object, FText& OutFailReason)
 		return false;
 	}
 
+	// Read the plugins
+	if(!FPluginReferenceDescriptor::ReadArray(Object, TEXT("Plugins"), Plugins, OutFailReason))
+	{
+		return false;
+	}
+
 	// Read the target platforms
 	const TArray< TSharedPtr<FJsonValue> > *TargetPlatformsValue;
 	if(Object.TryGetArrayField(TEXT("TargetPlatforms"), TargetPlatformsValue))
@@ -142,6 +148,9 @@ void FProjectDescriptor::Write(TJsonWriter<>& Writer) const
 	// Write the module list
 	FModuleDescriptor::WriteArray(Writer, TEXT("Modules"), Modules);
 
+	// Write the plugin list
+	FPluginReferenceDescriptor::WriteArray(Writer, TEXT("Plugins"), Plugins);
+
 	// Write the target platforms
 	if ( TargetPlatforms.Num() > 0 )
 	{
diff --git a/Engine/Source/Runtime/Projects/Private/ProjectManager.cpp b/Engine/Source/Runtime/Projects/Private/ProjectManager.cpp
--- a/Engine/Source/Runtime/Projects/Private/ProjectManager.cpp
+++ b/Engine/Source/Runtime/Projects/Private/ProjectManager.cpp
@@ -8,6 +8,12 @@ DEFINE_LOG_CATEGORY_STATIC( LogProjectManager, Log, All );
 
 FProjectManager::FProjectManager()
 {
+	bRestartRequired = false;
+}
+
+const FProjectDescriptor* FProjectManager::GetCurrentProject() const
+{
+	return CurrentProject.Get();
 }
 
 bool FProjectManager::LoadProjectFile( const FString& InProjectFile )
@@ -18,7 +24,7 @@ bool FProjectManager::LoadProjectFile( const FString& InProjectFile )
 	if(Descriptor->Load(InProjectFile, FailureReason))
 	{
 		// Create the project
-		CurrentlyLoadedProject = Descriptor;
+		CurrentProject = Descriptor;
 		return true;
 	}
 	
@@ -39,10 +45,10 @@ bool FProjectManager::LoadModulesForProject( const ELoadingPhase::Type LoadingPh
 
 	bool bSuccess = true;
 
-	if ( CurrentlyLoadedProject.IsValid() )
+	if ( CurrentProject.IsValid() )
 	{
 		TMap<FName, EModuleLoadResult> ModuleLoadFailures;
-		FModuleDescriptor::LoadModulesForPhase(LoadingPhase, CurrentlyLoadedProject->Modules, ModuleLoadFailures);
+		FModuleDescriptor::LoadModulesForPhase(LoadingPhase, CurrentProject->Modules, ModuleLoadFailures);
 
 		if ( ModuleLoadFailures.Num() > 0 )
 		{
@@ -92,7 +98,7 @@ bool FProjectManager::LoadModulesForProject( const ELoadingPhase::Type LoadingPh
 
 bool FProjectManager::AreProjectModulesUpToDate()
 {
-	return !CurrentlyLoadedProject.IsValid() || FModuleDescriptor::AreModulesUpToDate(CurrentlyLoadedProject->Modules);
+	return !CurrentProject.IsValid() || FModuleDescriptor::AreModulesUpToDate(CurrentProject->Modules);
 }
 
 const FString& FProjectManager::GetAutoLoadProjectFileName()
@@ -148,26 +154,26 @@ bool FProjectManager::DuplicateProjectFile(const FString& SourceProjectFilename,
 
 bool FProjectManager::UpdateLoadedProjectFileToCurrent(const TArray<FString>* StartupModuleNames, const FString& EngineIdentifier, FText& OutFailReason)
 {
-	if ( !CurrentlyLoadedProject.IsValid() )
+	if ( !CurrentProject.IsValid() )
 	{
 		return false;
 	}
 
 	// Freshen version information
-	CurrentlyLoadedProject->EngineAssociation = EngineIdentifier;
+	CurrentProject->EngineAssociation = EngineIdentifier;
 
 	// Replace the modules names, if specified
 	if(StartupModuleNames != NULL)
 	{
-		CurrentlyLoadedProject->Modules.Empty();
+		CurrentProject->Modules.Empty();
 		for(int32 Idx = 0; Idx < StartupModuleNames->Num(); Idx++)
 		{
-			CurrentlyLoadedProject->Modules.Add(FModuleDescriptor(*(*StartupModuleNames)[Idx]));
+			CurrentProject->Modules.Add(FModuleDescriptor(*(*StartupModuleNames)[Idx]));
 		}
 	}
 
 	// Update file on disk
-	return CurrentlyLoadedProject->Save(FPaths::GetProjectFilePath(), OutFailReason);
+	return CurrentProject->Save(FPaths::GetProjectFilePath(), OutFailReason);
 }
 
 bool FProjectManager::SignSampleProject(const FString& FilePath, const FString& Category, FText& OutFailReason)
@@ -198,12 +204,12 @@ bool FProjectManager::QueryStatusForProject(const FString& FilePath, FProjectSta
 
 bool FProjectManager::QueryStatusForCurrentProject(FProjectStatus& OutProjectStatus) const
 {
-	if ( !CurrentlyLoadedProject.IsValid() )
+	if ( !CurrentProject.IsValid() )
 	{
 		return false;
 	}
 
-	QueryStatusForProjectImpl(*CurrentlyLoadedProject, FPaths::GetProjectFilePath(), OutProjectStatus);
+	QueryStatusForProjectImpl(*CurrentProject, FPaths::GetProjectFilePath(), OutProjectStatus);
 	return true;
 }
 
@@ -250,15 +256,15 @@ void FProjectManager::UpdateSupportedTargetPlatformsForProject(const FString& Fi
 
 void FProjectManager::UpdateSupportedTargetPlatformsForCurrentProject(const FName& InPlatformName, const bool bIsSupported)
 {
-	if ( !CurrentlyLoadedProject.IsValid() )
+	if ( !CurrentProject.IsValid() )
 	{
 		return;
 	}
 
-	CurrentlyLoadedProject->UpdateSupportedTargetPlatforms(InPlatformName, bIsSupported);
+	CurrentProject->UpdateSupportedTargetPlatforms(InPlatformName, bIsSupported);
 
 	FText FailReason;
-	CurrentlyLoadedProject->Save(FPaths::GetProjectFilePath(), FailReason);
+	CurrentProject->Save(FPaths::GetProjectFilePath(), FailReason);
 
 	OnTargetPlatformsForCurrentProjectChangedEvent.Broadcast();
 }
@@ -287,19 +293,160 @@ void FProjectManager::ClearSupportedTargetPlatformsForProject(const FString& Fil
 
 void FProjectManager::ClearSupportedTargetPlatformsForCurrentProject()
 {
-	if ( !CurrentlyLoadedProject.IsValid() )
+	if ( !CurrentProject.IsValid() )
 	{
 		return;
 	}
 
-	CurrentlyLoadedProject->TargetPlatforms.Empty();
+	CurrentProject->TargetPlatforms.Empty();
 
 	FText FailReason;
-	CurrentlyLoadedProject->Save(FPaths::GetProjectFilePath(), FailReason);
+	CurrentProject->Save(FPaths::GetProjectFilePath(), FailReason);
 
 	OnTargetPlatformsForCurrentProjectChangedEvent.Broadcast();
 }
 
+void FProjectManager::GetEnabledPlugins(TArray<FString>& OutPluginNames) const
+{
+	// Get the default list of plugin names
+	GetDefaultEnabledPlugins(OutPluginNames);
+
+	// Modify that with the list of plugins in the project file
+	const FProjectDescriptor *Project = GetCurrentProject();
+	if(Project != NULL)
+	{
+		for(const FPluginReferenceDescriptor& Plugin: Project->Plugins)
+		{
+			if(Plugin.IsEnabledForPlatform(FPlatformMisc::GetUBTPlatform()))
+			{
+				OutPluginNames.AddUnique(Plugin.Name);
+			}
+			else
+			{
+				OutPluginNames.Remove(Plugin.Name);
+			}
+		}
+	}
+}
+
+bool FProjectManager::IsThirdPartyPluginEnabled() const
+{
+	static bool bInit = false;
+	TArray<FString> StandardPlugins;
+	if (!bInit)
+	{
+		StandardPlugins.Add(TEXT("BlankPlugin"));
+		StandardPlugins.Add(TEXT("UObjectPlugin"));
+		StandardPlugins.Add(TEXT("PluginsEditor"));
+		StandardPlugins.Add(TEXT("EpicSurvey"));
+		StandardPlugins.Add(TEXT("OculusRift"));
+		StandardPlugins.Add(TEXT("CustomMeshComponent"));
+		StandardPlugins.Add(TEXT("MessagingDebugger"));
+		StandardPlugins.Add(TEXT("PerforceSourceControl"));
+		StandardPlugins.Add(TEXT("SubversionSourceControl"));
+		StandardPlugins.Add(TEXT("UdpMessaging"));
+		StandardPlugins.Add(TEXT("WindowsMoviePlayer"));
+		StandardPlugins.Add(TEXT("CableComponent"));
+		StandardPlugins.Add(TEXT("ExampleDeviceProfileSelector"));
+		StandardPlugins.Add(TEXT("WinDualShock"));
+		StandardPlugins.Add(TEXT("VisualStudioSourceCodeAccess"));
+		StandardPlugins.Add(TEXT("XCodeSourceCodeAccess"));
+		StandardPlugins.Add(TEXT("SlateRemote"));
+		StandardPlugins.Add(TEXT("ScriptPlugin"));
+		StandardPlugins.Add(TEXT("ScriptEditorPlugin"));
+		StandardPlugins.Add(TEXT("SpeedTreeImporter"));
+		StandardPlugins.Add(TEXT("AppleMoviePlayer"));
+		StandardPlugins.Add(TEXT("IOSDeviceProfileSelector"));
+		bInit = true;
+	}
+
+	TArray<FString> EnabledPlugins;
+	GetEnabledPlugins(EnabledPlugins);
+
+	for (int Index = 0; Index < EnabledPlugins.Num(); ++Index)
+	{
+		if (!StandardPlugins.Contains(EnabledPlugins[Index]))
+		{
+			return true;
+		}
+	}
+	return false;
+}
+
+bool FProjectManager::SetPluginEnabled(const FString& PluginName, bool bEnabled, FText& OutFailReason)
+{
+	// Don't go any further if there's no project loaded
+	if(!CurrentProject.IsValid())
+	{
+		OutFailReason = LOCTEXT("NoProjectLoaded", "No project is currently loaded");
+		return false;
+	}
+
+	// Find or create the index of any existing reference in the project descriptor
+	int PluginRefIdx = 0;
+	for(;;PluginRefIdx++)
+	{
+		if(PluginRefIdx == CurrentProject->Plugins.Num())
+		{
+			PluginRefIdx = CurrentProject->Plugins.Add(FPluginReferenceDescriptor(PluginName, bEnabled));
+			break;
+		}
+		else if(CurrentProject->Plugins[PluginRefIdx].Name == PluginName)
+		{
+			CurrentProject->Plugins[PluginRefIdx].bEnabled = bEnabled;
+			break;
+		}
+	}
+
+	// If the current plugin reference is the default, just remove it from the list
+	const FPluginReferenceDescriptor* PluginRef = &CurrentProject->Plugins[PluginRefIdx];
+	if(PluginRef->WhitelistPlatforms.Num() == 0 && PluginRef->BlacklistPlatforms.Num() == 0)
+	{
+		// Get the default list of enabled plugins
+		TArray<FString> DefaultEnabledPlugins;
+		GetDefaultEnabledPlugins(DefaultEnabledPlugins);
+
+		// Check the enabled state is the same in that
+		if(DefaultEnabledPlugins.Contains(PluginName) == bEnabled)
+		{
+			CurrentProject->Plugins.RemoveAt(PluginRefIdx);
+			PluginRefIdx = INDEX_NONE;
+		}
+	}
+
+	// Try to save the project file
+	if(!CurrentProject->Save(*FPaths::GetProjectFilePath(), OutFailReason))
+	{
+		return false;
+	}
+
+	// Flag that a restart is required and return
+	bRestartRequired = true;
+	return true;
+}
+
+bool FProjectManager::IsRestartRequired() const
+{
+	return bRestartRequired;
+}
+
+void FProjectManager::GetDefaultEnabledPlugins(TArray<FString>& OutPluginNames)
+{
+	// Read all the enabled plugins from the ini file
+	OutPluginNames.Empty();
+	GConfig->GetArray( TEXT("Plugins"), TEXT("EnabledPlugins"), OutPluginNames, GEngineIni );
+
+	// Add all the game plugins
+	TArray<FPluginStatus> PluginStatuses = IPluginManager::Get().QueryStatusForAllPlugins();
+	for(const FPluginStatus& PluginStatus: PluginStatuses)
+	{
+		if(!PluginStatus.bIsBuiltIn)
+		{
+			OutPluginNames.AddUnique(PluginStatus.Name);
+		}
+	}
+}
+
 IProjectManager& IProjectManager::Get()
 {
 	// Single instance of manager, allocated on demand and destroyed on program exit.
diff --git a/Engine/Source/Runtime/Projects/Private/ProjectManager.h b/Engine/Source/Runtime/Projects/Private/ProjectManager.h
--- a/Engine/Source/Runtime/Projects/Private/ProjectManager.h
+++ b/Engine/Source/Runtime/Projects/Private/ProjectManager.h
@@ -13,6 +13,7 @@ public:
 	FProjectManager();
 
 	/** IProjectManager interface */
+	virtual const FProjectDescriptor *GetCurrentProject() const override;
 	virtual bool LoadProjectFile( const FString& ProjectFile ) override;
 	virtual bool LoadModulesForProject( const ELoadingPhase::Type LoadingPhase ) override;
 	virtual bool AreProjectModulesUpToDate( ) override;
@@ -29,15 +30,25 @@ public:
 	virtual void ClearSupportedTargetPlatformsForProject(const FString& FilePath) override;
 	virtual void ClearSupportedTargetPlatformsForCurrentProject() override;
 	virtual FOnTargetPlatformsForCurrentProjectChangedEvent& OnTargetPlatformsForCurrentProjectChanged() override { return OnTargetPlatformsForCurrentProjectChangedEvent; }
+	virtual void GetEnabledPlugins(TArray<FString>& OutPluginNames) const override;
+	virtual bool IsThirdPartyPluginEnabled() const override;
+	virtual bool SetPluginEnabled(const FString& PluginName, bool bEnabled, FText& OutFailReason) override;
+	virtual bool IsRestartRequired() const override;
 
 private:
 	static void QueryStatusForProjectImpl(const FProjectDescriptor& Project, const FString& FilePath, FProjectStatus& OutProjectStatus);
 
+	/** Gets the list of plugins enabled by default, excluding the project overrides */
+	static void GetDefaultEnabledPlugins(TArray<FString>& PluginNames);
+
 	/** The project that is currently loaded in the editor */
-	TSharedPtr< FProjectDescriptor > CurrentlyLoadedProject;
+	TSharedPtr< FProjectDescriptor > CurrentProject;
 
 	/** Delegate called when the target platforms for the current project are changed */
 	FOnTargetPlatformsForCurrentProjectChangedEvent OnTargetPlatformsForCurrentProjectChangedEvent;
+
+	/** Whether we need to restart in order to reflect change to the current project */
+	bool bRestartRequired;
 };
 
 
diff --git a/Engine/Source/Runtime/Projects/Public/Interfaces/IPluginManager.h b/Engine/Source/Runtime/Projects/Public/Interfaces/IPluginManager.h
--- a/Engine/Source/Runtime/Projects/Public/Interfaces/IPluginManager.h
+++ b/Engine/Source/Runtime/Projects/Public/Interfaces/IPluginManager.h
@@ -110,21 +110,4 @@ public:
 	 * @return	 Array of plugin status objects
 	 */
 	virtual TArray< FPluginStatus > QueryStatusForAllPlugins() const = 0;
-
-	/**
-	 * Sets the enabled state for the specified plugin. This change may not take effect until the editor is restarted.
-	 */
-	virtual void SetPluginEnabled( const FString& PluginName, bool bEnabled ) = 0;
-
-	/**
-	 * @return true if plugin changes have been made that will only take effect after a restart
-	 */
-	virtual bool IsRestartRequired() const = 0;
-
-	/**
-	 * Determines if a third-party plugin has been requested in the executable
-	 *
-	 * @return true if a third party plugin is enabled
-	 */
-	virtual bool HasThirdPartyPlugin() const = 0;
 };
diff --git a/Engine/Source/Runtime/Projects/Public/Interfaces/IProjectManager.h b/Engine/Source/Runtime/Projects/Public/Interfaces/IProjectManager.h
--- a/Engine/Source/Runtime/Projects/Public/Interfaces/IProjectManager.h
+++ b/Engine/Source/Runtime/Projects/Public/Interfaces/IProjectManager.h
@@ -5,6 +5,7 @@
 
 #include "EngineVersion.h"
 #include "ModuleDescriptor.h"
+#include "ProjectDescriptor.h"
 
 /**
  * Simple data structure that is filled when querying information about projects
@@ -76,6 +77,13 @@ public:
 	 */
 	static PROJECTS_API IProjectManager& Get();
 
+	/**
+	 * Gets the current project descriptor.
+	 *
+	 * @return Pointer to the currently loaded project descriptor. NULL if none is loaded.
+	 */
+	virtual const FProjectDescriptor* GetCurrentProject() const = 0;
+
 	/**
 	 * Loads the specified project file.
 	 *
@@ -212,6 +220,37 @@ public:
 	DECLARE_MULTICAST_DELEGATE(FOnTargetPlatformsForCurrentProjectChangedEvent);
 	virtual FOnTargetPlatformsForCurrentProjectChangedEvent& OnTargetPlatformsForCurrentProjectChanged() = 0;
 
+	/**
+	 * Gets a list of plugins enabled for the current project.
+	 * 
+	 * @param	OutPluginNames		Array to receive the list of plugin names
+	 */
+	virtual void GetEnabledPlugins(TArray<FString>& OutPluginNames) const = 0;
+
+	/**
+	 * Hack to checks whether the current project has a third party plugin enabled (ie. one which is not included by default in UE4Game).
+	 * 
+	 * @return	True if the project has a third party plugin enabled.
+	 */
+	virtual bool IsThirdPartyPluginEnabled() const = 0;
+
+	/**
+	 * Sets whether a plugin is enabled, and updates the project descriptor on disk. May require restarting to load it.
+	 * 
+	 * @param	PluginName		Name of the plugin
+	 * @param	bEnabled		Whether to enable or disable the plugin
+	 * @param	OutFailReason	On failure, gives an error message
+	 * @return	True if the plugin has been marked as enabled, and the project descriptor has been saved.
+	 */
+	virtual bool SetPluginEnabled(const FString& PluginName, bool bEnabled, FText& OutFailReason) = 0;
+
+	/**
+	 * Returns whether a restart is required to reflect changes to the project.
+	 * 
+	 * @return	True if the application needs to be restarted.
+	 */
+	virtual bool IsRestartRequired() const = 0;
+
 	/** Helper functions to reduce the syntax complexity of commonly used functions */
 	static const FString& GetProjectFileExtension() { return Get().NonStaticGetProjectFileExtension(); }
 };
\ No newline at end of file
diff --git a/Engine/Source/Runtime/Projects/Public/PluginDescriptor.h b/Engine/Source/Runtime/Projects/Public/PluginDescriptor.h
--- a/Engine/Source/Runtime/Projects/Public/PluginDescriptor.h
+++ b/Engine/Source/Runtime/Projects/Public/PluginDescriptor.h
@@ -2,6 +2,7 @@
 
 #pragma once
 
+#include "Projects.h"
 #include "ModuleDescriptor.h"
 
 /**
@@ -73,3 +74,40 @@ struct PROJECTS_API FPluginDescriptor
 	/** Reads the descriptor from the given JSON object */
 	bool Read(const FJsonObject& Object, FText& OutFailReason);
 };
+
+/**
+ * Descriptor for a plugin reference. Contains the information required to enable or disable a plugin for a given platform.
+ */
+struct PROJECTS_API FPluginReferenceDescriptor
+{
+	/** Name of the plugin */
+	FString Name;
+
+	/** Whether it should be enabled by default */
+	bool bEnabled;
+
+	/** If enabled, list of platforms for which the plugin should be enabled (or all platforms if blank). */
+	TArray<FString> WhitelistPlatforms;
+
+	/** If enabled, list of platforms for which the plugin should be disabled. */
+	TArray<FString> BlacklistPlatforms;
+
+	/** Constructor */
+	FPluginReferenceDescriptor(const FString& InName = TEXT(""), bool bInEnabled = false);
+
+	/** Determines whether the plugin is enabled for the given platform */
+	bool IsEnabledForPlatform(const FString& Platform) const;
+
+	/** Reads the descriptor from the given JSON object */
+	bool Read(const FJsonObject& Object, FText& OutFailReason);
+
+	/** Reads an array of modules from the given JSON object */
+	static bool ReadArray(const FJsonObject& Object, const TCHAR* Name, TArray<FPluginReferenceDescriptor>& OutModules, FText& OutFailReason);
+
+	/** Writes a descriptor to JSON */
+	void Write(TJsonWriter<>& Writer) const;
+
+	/** Writes an array of modules to JSON */
+	static void WriteArray(TJsonWriter<>& Writer, const TCHAR* Name, const TArray<FPluginReferenceDescriptor>& Modules);
+};
+
diff --git a/Engine/Source/Runtime/Projects/Public/ProjectDescriptor.h b/Engine/Source/Runtime/Projects/Public/ProjectDescriptor.h
--- a/Engine/Source/Runtime/Projects/Public/ProjectDescriptor.h
+++ b/Engine/Source/Runtime/Projects/Public/ProjectDescriptor.h
@@ -3,6 +3,8 @@
 #pragma once
 
 #include "Projects.h"
+#include "ModuleDescriptor.h"
+#include "PluginDescriptor.h"
 
 /**
  * Version numbers for project descriptors.
@@ -40,9 +42,12 @@ struct PROJECTS_API FProjectDescriptor
 	/** Description to show in the project browser */
 	FString Description;
 
-	/** List of all modules associated with this plugin */
+	/** List of all modules associated with this project */
 	TArray<FModuleDescriptor> Modules;
 
+	/** List of plugins for this project (may be enabled/disabled) */
+	TArray<FPluginReferenceDescriptor> Plugins;
+
 	/** Array of platforms that this project is targeting */
 	TArray<FName> TargetPlatforms;
 
