Prevent an ensure when starting programs that compile-in the plugin manager but do not support content plugins, like SCW.

diff --git a/Engine/Source/Runtime/Projects/Private/PluginManager.cpp b/Engine/Source/Runtime/Projects/Private/PluginManager.cpp
--- a/Engine/Source/Runtime/Projects/Private/PluginManager.cpp
+++ b/Engine/Source/Runtime/Projects/Private/PluginManager.cpp
@@ -298,7 +298,7 @@ void FPluginManager::EnablePluginsThatAreConfiguredToBeEnabled()
 	}
 
 	// Mount all the plugin content folders
-	if( ensure( RegisterMountPointDelegate.IsBound() ) )
+	if( ContentFolders.Num() > 0 && ensure( RegisterMountPointDelegate.IsBound() ) )
 	{
 		for(const FPluginContentFolder& ContentFolder: ContentFolders)
 		{
