Don't tell the module manager about plugin modules before we need to load them. It shouldn't need to know, and it causes it to cache when modules are missing even if they can be compiled.

diff --git a/Engine/Source/Runtime/Projects/Private/PluginManager.cpp b/Engine/Source/Runtime/Projects/Private/PluginManager.cpp
--- a/Engine/Source/Runtime/Projects/Private/PluginManager.cpp
+++ b/Engine/Source/Runtime/Projects/Private/PluginManager.cpp
@@ -34,7 +34,7 @@ FPluginInstance::FPluginInstance(const FString& InFileName, const FPluginDescrip
 
 
 FPluginManager::FPluginManager()
-	: bEarliestPhaseProcessed(false)
+	: bHaveConfiguredEnabledPlugins(false)
 {
 	DiscoverAllPlugins();
 }
@@ -246,74 +246,66 @@ void FPluginManager::DiscoverAllPlugins()
 			// Add the plugin binaries directory
 			const FString PluginBinariesPath = FPaths::Combine(*FPaths::GetPath(Plugin->FileName), TEXT("Binaries"), FPlatformProcess::GetBinariesSubdirectory());
 			FModuleManager::Get().AddBinariesDirectory(*PluginBinariesPath, Plugin->LoadedFrom == EPluginLoadedFrom::GameProject);
-
-			// Make sure to tell the module manager about any of our plugin's modules, so it will know where on disk to find them.
-			for( auto ModuleInfoIt( Plugin->Descriptor.Modules.CreateConstIterator() ); ModuleInfoIt; ++ModuleInfoIt )
-			{
-				const FModuleDescriptor& ModuleInfo = *ModuleInfoIt;
-
-				// @todo plugin: We're adding all modules always here, even editor modules which might not exist at runtime. Hope that's OK.
-				FModuleManager::Get().AddModule(ModuleInfo.Name);
-
-				// Add to the module->plugin map
-				ModulePluginMap.Add( ModuleInfo.Name, Plugin );
-			}
 		}
 	}
 }
 
-void FPluginManager::EnablePluginsThatAreConfiguredToBeEnabled()
+void FPluginManager::ConfigureEnabledPlugins()
 {
-	TArray< FString > EnabledPluginNames;
-#if IS_PROGRAM
-	GConfig->GetArray(TEXT("Plugins"), TEXT("ProgramEnabledPlugins"), EnabledPluginNames, GEngineIni);
-#else
-	FProjectManager::Get().GetEnabledPlugins(EnabledPluginNames);
-#endif
-
-	TSet< FString > AllEnabledPlugins;
-	AllEnabledPlugins.Append( MoveTemp(EnabledPluginNames) );
-
-	// Enable all the plugins by name
-	for( const TSharedRef< FPluginInstance > Plugin : AllPlugins )
+	if(!bHaveConfiguredEnabledPlugins)
 	{
-		if ( AllEnabledPlugins.Contains(Plugin->Name) )
+		TArray< FString > EnabledPluginNames;
+		#if IS_PROGRAM
+			GConfig->GetArray(TEXT("Plugins"), TEXT("ProgramEnabledPlugins"), EnabledPluginNames, GEngineIni);
+		#else
+			FProjectManager::Get().GetEnabledPlugins(EnabledPluginNames);
+		#endif
+
+		// Build a set from the array
+		TSet< FString > AllEnabledPlugins;
+		AllEnabledPlugins.Append( MoveTemp(EnabledPluginNames) );
+
+		// Enable all the plugins by name
+		for( const TSharedRef< FPluginInstance > Plugin : AllPlugins )
 		{
-			Plugin->bEnabled = true;
+			if ( AllEnabledPlugins.Contains(Plugin->Name) )
+			{
+				Plugin->bEnabled = true;
+			}
 		}
-	}
 
-	// Build the list of content folders
-	ContentFolders.Empty();
-	for(const TSharedRef<FPluginInstance>& Plugin: AllPlugins)
-	{
-		if(Plugin->bEnabled && Plugin->Descriptor.bCanContainContent)
+		// Build the list of content folders
+		ContentFolders.Empty();
+		for(const TSharedRef<FPluginInstance>& Plugin: AllPlugins)
 		{
-			FPluginContentFolder ContentFolder;
-			ContentFolder.Name = Plugin->Name;
-			ContentFolder.RootPath = FString::Printf(TEXT("/%s/"), *Plugin->Name);
-			ContentFolder.ContentPath = FPaths::GetPath(Plugin->FileName) / TEXT("Content");
-			ContentFolders.Emplace(ContentFolder);
+			if(Plugin->bEnabled && Plugin->Descriptor.bCanContainContent)
+			{
+				FPluginContentFolder ContentFolder;
+				ContentFolder.Name = Plugin->Name;
+				ContentFolder.RootPath = FString::Printf(TEXT("/%s/"), *Plugin->Name);
+				ContentFolder.ContentPath = FPaths::GetPath(Plugin->FileName) / TEXT("Content");
+				ContentFolders.Emplace(ContentFolder);
+			}
 		}
-	}
 
-	// Mount all the plugin content folders
-	if( ContentFolders.Num() > 0 && ensure( RegisterMountPointDelegate.IsBound() ) )
-	{
-		for(const FPluginContentFolder& ContentFolder: ContentFolders)
+		// Mount all the plugin content folders
+		if( ContentFolders.Num() > 0 && ensure( RegisterMountPointDelegate.IsBound() ) )
 		{
-			RegisterMountPointDelegate.Execute(ContentFolder.RootPath, ContentFolder.ContentPath);
+			for(const FPluginContentFolder& ContentFolder: ContentFolders)
+			{
+				RegisterMountPointDelegate.Execute(ContentFolder.RootPath, ContentFolder.ContentPath);
+			}
 		}
+
+		// Don't need to run this again
+		bHaveConfiguredEnabledPlugins = true;
 	}
 }
 
 void FPluginManager::LoadModulesForEnabledPlugins( const ELoadingPhase::Type LoadingPhase )
 {
-	if ( !bEarliestPhaseProcessed )
-	{
-		EnablePluginsThatAreConfiguredToBeEnabled();
-		bEarliestPhaseProcessed = true;
-	}
+	// Figure out which plugins are enabled
+	ConfigureEnabledPlugins();
 
 	// Load plugins!
 	for( const TSharedRef< FPluginInstance > Plugin : AllPlugins )
@@ -374,18 +366,9 @@ void FPluginManager::SetRegisterMountPointDelegate( const FRegisterMountPointDel
 	RegisterMountPointDelegate = Delegate;
 }
 
-bool FPluginManager::IsPluginModule( const FName ModuleName ) const
-{
-	return ( ModulePluginMap.Find( ModuleName ) != NULL );
-}
-
 bool FPluginManager::AreEnabledPluginModulesUpToDate()
 {
-	if (!bEarliestPhaseProcessed)
-	{
-		EnablePluginsThatAreConfiguredToBeEnabled();
-		bEarliestPhaseProcessed = true;
-	}
+	ConfigureEnabledPlugins();
 
 	for (TArray< TSharedRef< FPluginInstance > >::TConstIterator Iter(AllPlugins); Iter; ++Iter)
 	{
diff --git a/Engine/Source/Runtime/Projects/Private/PluginManager.h b/Engine/Source/Runtime/Projects/Private/PluginManager.h
--- a/Engine/Source/Runtime/Projects/Private/PluginManager.h
+++ b/Engine/Source/Runtime/Projects/Private/PluginManager.h
@@ -61,7 +61,6 @@ public:
 	/** IPluginManager interface */
 	virtual void LoadModulesForEnabledPlugins( const ELoadingPhase::Type LoadingPhase ) override;
 	virtual void SetRegisterMountPointDelegate( const FRegisterMountPointDelegate& Delegate ) override;
-	virtual bool IsPluginModule( const FName ModuleName ) const override;
 	virtual bool AreEnabledPluginModulesUpToDate() override;
 	virtual TArray< FPluginStatus > QueryStatusForAllPlugins() const override;
 	virtual const TArray< FPluginContentFolder >& GetPluginContentFolders() const override;
@@ -73,15 +72,12 @@ private:
 	void DiscoverAllPlugins();
 
 	/** Sets the bPluginEnabled flag on all plugins found from DiscoverAllPlugins that are enabled in config */
-	void EnablePluginsThatAreConfiguredToBeEnabled();
+	void ConfigureEnabledPlugins();
 
 private:
 	/** All of the plugins that we know about */
 	TArray< TSharedRef< FPluginInstance > > AllPlugins;
 
-	/** Map from module name to plugin containing that module */
-	TMap< FName, TSharedRef< FPluginInstance > > ModulePluginMap;
-
 	/** All the plugin content folders */
 	TArray<FPluginContentFolder> ContentFolders;
 
@@ -89,8 +85,8 @@ private:
 	    content path mounting functionality from Core. */
 	FRegisterMountPointDelegate RegisterMountPointDelegate;
 
-	/** The earliest phase that modules are loaded requires some additional initialization */
-	bool bEarliestPhaseProcessed;
+	/** Set when all the appropriate plugins have been marked as enabled */
+	bool bHaveConfiguredEnabledPlugins;
 };
 
 
diff --git a/Engine/Source/Runtime/Projects/Public/Interfaces/IPluginManager.h b/Engine/Source/Runtime/Projects/Public/Interfaces/IPluginManager.h
--- a/Engine/Source/Runtime/Projects/Public/Interfaces/IPluginManager.h
+++ b/Engine/Source/Runtime/Projects/Public/Interfaces/IPluginManager.h
@@ -111,14 +111,6 @@ public:
 	 */
 	virtual void SetRegisterMountPointDelegate( const FRegisterMountPointDelegate& Delegate ) = 0;
 
-	/** 
-	 * Queries whether a module is provided by a plugin
-	 *
-	 * @param	ModuleName	The name of the module to check
-	 * @returns true if the module is supplied by a plugin
-	 */
-	virtual bool IsPluginModule( const FName ModuleName ) const = 0;
-
 	/** 
 	 * Checks whether modules for the enabled plugins are up to date
 	 *
