Add a PostEngineInit load phase for modules

diff --git a/Engine/Source/Runtime/Launch/Private/LaunchEngineLoop.cpp b/Engine/Source/Runtime/Launch/Private/LaunchEngineLoop.cpp
--- a/Engine/Source/Runtime/Launch/Private/LaunchEngineLoop.cpp
+++ b/Engine/Source/Runtime/Launch/Private/LaunchEngineLoop.cpp
@@ -1866,6 +1866,14 @@ int32 FEngineLoop::Init()
 
 	GEngine->Init(this);
 
+	// Load all the post-engine init modules
+	if (!IProjectManager::Get().LoadModulesForProject(ELoadingPhase::PostEngineInit))
+	{
+		GIsRequestingExit = true;
+		return 1;
+	}
+	IPluginManager::Get().LoadModulesForEnabledPlugins(ELoadingPhase::PostEngineInit);
+
 	GetMoviePlayer()->WaitForMovieToFinish();
 
 	// initialize automation worker
diff --git a/Engine/Source/Runtime/Projects/Private/ModuleDescriptor.cpp b/Engine/Source/Runtime/Projects/Private/ModuleDescriptor.cpp
--- a/Engine/Source/Runtime/Projects/Private/ModuleDescriptor.cpp
+++ b/Engine/Source/Runtime/Projects/Private/ModuleDescriptor.cpp
@@ -37,6 +37,9 @@ const TCHAR* ELoadingPhase::ToString( const ELoadingPhase::Type Value )
 	case PreLoadingScreen:
 		return TEXT( "PreLoadingScreen" );
 
+	case PostEngineInit:
+		return TEXT( "PostEngineInit" );
+
 	default:
 		ensureMsgf( false, TEXT( "Unrecognized ELoadingPhase value: %i" ), Value );
 		return NULL;
diff --git a/Engine/Source/Runtime/Projects/Public/ModuleDescriptor.h b/Engine/Source/Runtime/Projects/Public/ModuleDescriptor.h
--- a/Engine/Source/Runtime/Projects/Public/ModuleDescriptor.h
+++ b/Engine/Source/Runtime/Projects/Public/ModuleDescriptor.h
@@ -27,8 +27,11 @@ namespace ELoadingPhase
 
 		/** Loaded before the engine is fully initialized for modules that need to hook into the loading screen before it triggers */
 		PreLoadingScreen,
-		// NOTE: If you add a new value, make sure to update the ToString() method below!
 
+		/** After the engine has been initialized */
+		PostEngineInit,
+
+		// NOTE: If you add a new value, make sure to update the ToString() method below!
 		Max
 	};
 
