Require all plugins to be present for a project to load, rather than attempting to carry on without them. Plugins may define object types, and loading packages without them can cause corruption. Can revisit in the future to allow for 'editor-only' plugins.

Plugin reference in the .uproject file has a "Description" field which allows setting a customizable message if the plugin is not available. Allows distribution of samples which require a plugin, but without including the plugin itself.

diff --git a/Engine/Source/Runtime/Launch/Private/LaunchEngineLoop.cpp b/Engine/Source/Runtime/Launch/Private/LaunchEngineLoop.cpp
--- a/Engine/Source/Runtime/Launch/Private/LaunchEngineLoop.cpp
+++ b/Engine/Source/Runtime/Launch/Private/LaunchEngineLoop.cpp
@@ -1295,14 +1295,11 @@ int32 FEngineLoop::PreInit( const TCHAR* CmdLine )
 #endif
 
 	// Load up all modules that need to hook into the loading screen
-	if (!IProjectManager::Get().LoadModulesForProject(ELoadingPhase::PreLoadingScreen))
+	if (!IProjectManager::Get().LoadModulesForProject(ELoadingPhase::PreLoadingScreen) || !IPluginManager::Get().LoadModulesForEnabledPlugins(ELoadingPhase::PreLoadingScreen))
 	{
 		return 1;
 	}
 
-	IPluginManager::Get().LoadModulesForEnabledPlugins(ELoadingPhase::PreLoadingScreen);
-	
-	
 #if !UE_SERVER
 	if ( !IsRunningDedicatedServer() )
 	{
@@ -1773,29 +1770,23 @@ bool FEngineLoop::LoadStartupCoreModules()
 bool FEngineLoop::LoadStartupModules()
 {
 	// Load any modules that want to be loaded before default modules are loaded up.
-	if (!IProjectManager::Get().LoadModulesForProject(ELoadingPhase::PreDefault))
+	if (!IProjectManager::Get().LoadModulesForProject(ELoadingPhase::PreDefault) || !IPluginManager::Get().LoadModulesForEnabledPlugins(ELoadingPhase::PreDefault))
 	{
 		return false;
 	}
 
-	IPluginManager::Get().LoadModulesForEnabledPlugins(ELoadingPhase::PreDefault);
-
 	// Load modules that are configured to load in the default phase
-	if (!IProjectManager::Get().LoadModulesForProject(ELoadingPhase::Default))
+	if (!IProjectManager::Get().LoadModulesForProject(ELoadingPhase::Default) || !IPluginManager::Get().LoadModulesForEnabledPlugins(ELoadingPhase::Default))
 	{
 		return false;
 	}
 
-	IPluginManager::Get().LoadModulesForEnabledPlugins(ELoadingPhase::Default);
-
 	// Load any modules that want to be loaded after default modules are loaded up.
-	if (!IProjectManager::Get().LoadModulesForProject(ELoadingPhase::PostDefault))
+	if (!IProjectManager::Get().LoadModulesForProject(ELoadingPhase::PostDefault) || !IPluginManager::Get().LoadModulesForEnabledPlugins(ELoadingPhase::PostDefault))
 	{
 		return false;
 	}
 
-	IPluginManager::Get().LoadModulesForEnabledPlugins(ELoadingPhase::PostDefault);
-
 	return true;
 }
 
@@ -1884,12 +1875,11 @@ int32 FEngineLoop::Init()
 	GEngine->Init(this);
 
 	// Load all the post-engine init modules
-	if (!IProjectManager::Get().LoadModulesForProject(ELoadingPhase::PostEngineInit))
+	if (!IProjectManager::Get().LoadModulesForProject(ELoadingPhase::PostEngineInit) || !IPluginManager::Get().LoadModulesForEnabledPlugins(ELoadingPhase::PostEngineInit))
 	{
 		GIsRequestingExit = true;
 		return 1;
 	}
-	IPluginManager::Get().LoadModulesForEnabledPlugins(ELoadingPhase::PostEngineInit);
 
 	GetMoviePlayer()->WaitForMovieToFinish();
 
@@ -2342,7 +2332,7 @@ bool FEngineLoop::AppInit( )
 
 	// Check whether the project or any of its plugins are missing or are out of date
 #if UE_EDITOR
-	if(!GIsBuildMachine && FPaths::IsProjectFilePathSet())
+	if(!GIsBuildMachine && FPaths::IsProjectFilePathSet() && IPluginManager::Get().AreRequiredPluginsAvailable())
 	{
 		if(!IProjectManager::Get().AreProjectModulesUpToDate() || !IPluginManager::Get().AreEnabledPluginModulesUpToDate())
 		{
@@ -2374,9 +2364,9 @@ bool FEngineLoop::AppInit( )
 #endif
 
 	// Load "pre-init" plugin modules
-	if (IProjectManager::Get().LoadModulesForProject(ELoadingPhase::PostConfigInit))
+	if (!IProjectManager::Get().LoadModulesForProject(ELoadingPhase::PostConfigInit) || !IPluginManager::Get().LoadModulesForEnabledPlugins(ELoadingPhase::PostConfigInit))
 	{
-		IPluginManager::Get().LoadModulesForEnabledPlugins(ELoadingPhase::PostConfigInit);
+		return false;
 	}
 
 	// Put the command line and config info into the suppression system
diff --git a/Engine/Source/Runtime/Projects/Private/PluginDescriptor.cpp b/Engine/Source/Runtime/Projects/Private/PluginDescriptor.cpp
--- a/Engine/Source/Runtime/Projects/Private/PluginDescriptor.cpp
+++ b/Engine/Source/Runtime/Projects/Private/PluginDescriptor.cpp
@@ -145,6 +145,9 @@ bool FPluginReferenceDescriptor::Read(const FJsonObject& Object, FText& OutFailR
 		return false;
 	}
 
+	// Read the description
+	Object.TryGetStringField(TEXT("Description"), Description);
+
 	// Get the platform lists
 	Object.TryGetStringArrayField(TEXT("WhitelistPlatforms"), WhitelistPlatforms);
 	Object.TryGetStringArrayField(TEXT("BlacklistPlatforms"), BlacklistPlatforms);
@@ -178,6 +181,10 @@ void FPluginReferenceDescriptor::Write(TJsonWriter<>& Writer) const
 	Writer.WriteObjectStart();
 	Writer.WriteValue(TEXT("Name"), Name);
 	Writer.WriteValue(TEXT("Enabled"), bEnabled);
+	if (Description.Len() > 0)
+	{
+		Writer.WriteValue(TEXT("Description"), Description);
+	}
 	if (WhitelistPlatforms.Num() > 0)
 	{
 		Writer.WriteArrayStart(TEXT("WhitelistPlatforms"));
diff --git a/Engine/Source/Runtime/Projects/Private/PluginManager.cpp b/Engine/Source/Runtime/Projects/Private/PluginManager.cpp
--- a/Engine/Source/Runtime/Projects/Private/PluginManager.cpp
+++ b/Engine/Source/Runtime/Projects/Private/PluginManager.cpp
@@ -35,6 +35,7 @@ FPluginInstance::FPluginInstance(const FString& InFileName, const FPluginDescrip
 
 FPluginManager::FPluginManager()
 	: bHaveConfiguredEnabledPlugins(false)
+	, bHaveAllRequiredPlugins(false)
 {
 	DiscoverAllPlugins();
 }
@@ -250,10 +251,33 @@ void FPluginManager::DiscoverAllPlugins()
 	}
 }
 
-void FPluginManager::ConfigureEnabledPlugins()
+bool FPluginManager::ConfigureEnabledPlugins()
 {
 	if(!bHaveConfiguredEnabledPlugins)
 	{
+		// Don't need to run this again
+		bHaveConfiguredEnabledPlugins = true;
+
+		// If a current project is set, check that we know about any plugin that's explicitly enabled
+		const FProjectDescriptor *Project = IProjectManager::Get().GetCurrentProject();
+		if(Project != nullptr)
+		{
+			for(const FPluginReferenceDescriptor& Plugin: Project->Plugins)
+			{
+				if(Plugin.bEnabled && !FindPluginInstance(Plugin.Name).IsValid())
+				{
+					FText Caption(LOCTEXT("PluginMissingCaption", "Plugin missing"));
+					FString Description = (Plugin.Description.Len() > 0)? FString::Printf(TEXT("\n\n%s"), *Plugin.Description) : FString();
+					FMessageDialog::Open(EAppMsgType::Ok, FText::Format(LOCTEXT("PluginMissingError", "This project requires the {0} plugin. {1}"), FText::FromString(Plugin.Name), FText::FromString(Description)), &Caption);
+					return false;
+				}
+			}
+		}
+
+		// If we made it here, we have all the required plugins
+		bHaveAllRequiredPlugins = true;
+
+		// Get all the enabled plugin names
 		TArray< FString > EnabledPluginNames;
 		#if IS_PROGRAM
 			GConfig->GetArray(TEXT("Plugins"), TEXT("ProgramEnabledPlugins"), EnabledPluginNames, GEngineIni);
@@ -296,16 +320,31 @@ void FPluginManager::ConfigureEnabledPlugins()
 				RegisterMountPointDelegate.Execute(ContentFolder.RootPath, ContentFolder.ContentPath);
 			}
 		}
+	}
+	return bHaveAllRequiredPlugins;
+}
 
-		// Don't need to run this again
-		bHaveConfiguredEnabledPlugins = true;
+TSharedPtr<FPluginInstance> FPluginManager::FindPluginInstance(const FString& Name)
+{
+	TSharedPtr<FPluginInstance> Result;
+	for(const TSharedRef<FPluginInstance>& Instance : AllPlugins)
+	{
+		if(Instance->Name == Name)
+		{
+			Result = Instance;
+			break;
+		}
 	}
+	return Result;
 }
 
-void FPluginManager::LoadModulesForEnabledPlugins( const ELoadingPhase::Type LoadingPhase )
+bool FPluginManager::LoadModulesForEnabledPlugins( const ELoadingPhase::Type LoadingPhase )
 {
 	// Figure out which plugins are enabled
-	ConfigureEnabledPlugins();
+	if(!ConfigureEnabledPlugins())
+	{
+		return false;
+	}
 
 	// Load plugins!
 	for( const TSharedRef< FPluginInstance > Plugin : AllPlugins )
@@ -328,19 +367,19 @@ void FPluginManager::LoadModulesForEnabledPlugins( const ELoadingPhase::Type Loa
 
 					if ( FailureReason == EModuleLoadResult::FileNotFound )
 					{
-						FailureMessage = FText::Format( LOCTEXT("PluginModuleNotFound", "Plugin '{0}' failed to load because module '{1}' could not be found.  This plugin's functionality will not be available.  Please ensure the plugin is properly installed, otherwise consider disabling the plugin for this project."), PluginNameText, TextModuleName );
+						FailureMessage = FText::Format( LOCTEXT("PluginModuleNotFound", "Plugin '{0}' failed to load because module '{1}' could not be found.  Please ensure the plugin is properly installed, otherwise consider disabling the plugin for this project."), PluginNameText, TextModuleName );
 					}
 					else if ( FailureReason == EModuleLoadResult::FileIncompatible )
 					{
-						FailureMessage = FText::Format( LOCTEXT("PluginModuleIncompatible", "Plugin '{0}' failed to load because module '{1}' does not appear to be compatible with the current version of the engine.  This plugin's functionality will not be available.  The plugin may need to be recompiled."), PluginNameText, TextModuleName );
+						FailureMessage = FText::Format( LOCTEXT("PluginModuleIncompatible", "Plugin '{0}' failed to load because module '{1}' does not appear to be compatible with the current version of the engine.  The plugin may need to be recompiled."), PluginNameText, TextModuleName );
 					}
 					else if ( FailureReason == EModuleLoadResult::CouldNotBeLoadedByOS )
 					{
-						FailureMessage = FText::Format( LOCTEXT("PluginModuleCouldntBeLoaded", "Plugin '{0}' failed to load because module '{1}' could not be loaded.  This plugin's functionality will not be available.  There may be an operating system error or the module may not be properly set up."), PluginNameText, TextModuleName );
+						FailureMessage = FText::Format( LOCTEXT("PluginModuleCouldntBeLoaded", "Plugin '{0}' failed to load because module '{1}' could not be loaded.  There may be an operating system error or the module may not be properly set up."), PluginNameText, TextModuleName );
 					}
 					else if ( FailureReason == EModuleLoadResult::FailedToInitialize )
 					{
-						FailureMessage = FText::Format( LOCTEXT("PluginModuleFailedToInitialize", "Plugin '{0}' failed to load because module '{1}' could be initialized successfully after it was loaded.  This plugin's functionality will not be available."), PluginNameText, TextModuleName );
+						FailureMessage = FText::Format( LOCTEXT("PluginModuleFailedToInitialize", "Plugin '{0}' failed to load because module '{1}' could be initialized successfully after it was loaded."), PluginNameText, TextModuleName );
 					}
 					else 
 					{
@@ -356,9 +395,11 @@ void FPluginManager::LoadModulesForEnabledPlugins( const ELoadingPhase::Type Loa
 			if( !FailureMessage.IsEmpty() )
 			{
 				FMessageDialog::Open(EAppMsgType::Ok, FailureMessage);
+				return false;
 			}
 		}
 	}
+	return true;
 }
 
 void FPluginManager::SetRegisterMountPointDelegate( const FRegisterMountPointDelegate& Delegate )
@@ -366,9 +407,17 @@ void FPluginManager::SetRegisterMountPointDelegate( const FRegisterMountPointDel
 	RegisterMountPointDelegate = Delegate;
 }
 
+bool FPluginManager::AreRequiredPluginsAvailable()
+{
+	return ConfigureEnabledPlugins();
+}
+
 bool FPluginManager::AreEnabledPluginModulesUpToDate()
 {
-	ConfigureEnabledPlugins();
+	if(!ConfigureEnabledPlugins())
+	{
+		return false;
+	}
 
 	for (TArray< TSharedRef< FPluginInstance > >::TConstIterator Iter(AllPlugins); Iter; ++Iter)
 	{
diff --git a/Engine/Source/Runtime/Projects/Private/PluginManager.h b/Engine/Source/Runtime/Projects/Private/PluginManager.h
--- a/Engine/Source/Runtime/Projects/Private/PluginManager.h
+++ b/Engine/Source/Runtime/Projects/Private/PluginManager.h
@@ -59,8 +59,9 @@ public:
 	~FPluginManager();
 
 	/** IPluginManager interface */
-	virtual void LoadModulesForEnabledPlugins( const ELoadingPhase::Type LoadingPhase ) override;
+	virtual bool LoadModulesForEnabledPlugins( const ELoadingPhase::Type LoadingPhase ) override;
 	virtual void SetRegisterMountPointDelegate( const FRegisterMountPointDelegate& Delegate ) override;
+	virtual bool AreRequiredPluginsAvailable() override;
 	virtual bool AreEnabledPluginModulesUpToDate() override;
 	virtual TArray< FPluginStatus > QueryStatusForAllPlugins() const override;
 	virtual const TArray< FPluginContentFolder >& GetPluginContentFolders() const override;
@@ -72,7 +73,10 @@ private:
 	void DiscoverAllPlugins();
 
 	/** Sets the bPluginEnabled flag on all plugins found from DiscoverAllPlugins that are enabled in config */
-	void ConfigureEnabledPlugins();
+	bool ConfigureEnabledPlugins();
+
+	/** Gets the instance of a given plugin */
+	TSharedPtr<FPluginInstance> FindPluginInstance(const FString& Name);
 
 private:
 	/** All of the plugins that we know about */
@@ -87,6 +91,9 @@ private:
 
 	/** Set when all the appropriate plugins have been marked as enabled */
 	bool bHaveConfiguredEnabledPlugins;
+
+	/** Set if all the required plugins are available */
+	bool bHaveAllRequiredPlugins;
 };
 
 
diff --git a/Engine/Source/Runtime/Projects/Public/Interfaces/IPluginManager.h b/Engine/Source/Runtime/Projects/Public/Interfaces/IPluginManager.h
--- a/Engine/Source/Runtime/Projects/Public/Interfaces/IPluginManager.h
+++ b/Engine/Source/Runtime/Projects/Public/Interfaces/IPluginManager.h
@@ -98,7 +98,7 @@ public:
 	 * @param	LoadingPhase	Which loading phase we're loading plugin modules from.  Only modules that are configured to be
 	 *							loaded at the specified loading phase will be loaded during this call.
 	 */
-	virtual void LoadModulesForEnabledPlugins( const ELoadingPhase::Type LoadingPhase ) = 0;
+	virtual bool LoadModulesForEnabledPlugins( const ELoadingPhase::Type LoadingPhase ) = 0;
 
 	/** Delegate type for mounting content paths.  Used internally by FPackageName code. */
 	DECLARE_DELEGATE_TwoParams( FRegisterMountPointDelegate, const FString& /* Root content path */, const FString& /* Directory name */ );
@@ -111,6 +111,13 @@ public:
 	 */
 	virtual void SetRegisterMountPointDelegate( const FRegisterMountPointDelegate& Delegate ) = 0;
 
+	/**
+	 * Checks if all the required plugins are available. If not, will present an error dialog the first time a plugin is loaded or this function is called.
+	 *
+	 * @returns true if all the required plugins are available.
+	 */
+	virtual bool AreRequiredPluginsAvailable( ) = 0;
+
 	/** 
 	 * Checks whether modules for the enabled plugins are up to date
 	 *
diff --git a/Engine/Source/Runtime/Projects/Public/PluginDescriptor.h b/Engine/Source/Runtime/Projects/Public/PluginDescriptor.h
--- a/Engine/Source/Runtime/Projects/Public/PluginDescriptor.h
+++ b/Engine/Source/Runtime/Projects/Public/PluginDescriptor.h
@@ -89,6 +89,9 @@ struct PROJECTS_API FPluginReferenceDescriptor
 	/** Whether it should be enabled by default */
 	bool bEnabled;
 
+	/** Description of the plugin for users that do not have it installed. */
+	FString Description;
+
 	/** If enabled, list of platforms for which the plugin should be enabled (or all platforms if blank). */
 	TArray<FString> WhitelistPlatforms;
 
