#UE4  Respect Blacklist Platforms for Plugins.

                fixes crashes caused by early outs and uninitialized modules.

diff --git a/Engine/Source/Runtime/Projects/Private/PluginDescriptor.cpp b/Engine/Source/Runtime/Projects/Private/PluginDescriptor.cpp
--- a/Engine/Source/Runtime/Projects/Private/PluginDescriptor.cpp
+++ b/Engine/Source/Runtime/Projects/Private/PluginDescriptor.cpp
@@ -88,6 +88,12 @@ bool FPluginDescriptor::Read(const FJsonObject& Object, FText& OutFailReason)
 	}
 
 	Object.TryGetBoolField(TEXT("EnabledByDefault"), bEnabledByDefault);
+
+	for ( FModuleDescriptor Descriptor : Modules )
+	{
+		bEnabledByDefault &= (!Descriptor.BlacklistPlatforms.Contains(FString(FPlatformProperties::PlatformName())));
+	}
+
 	Object.TryGetBoolField(TEXT("CanContainContent"), bCanContainContent);
 	Object.TryGetBoolField(TEXT("IsBetaVersion"), bIsBetaVersion);
 	return true;
