Fixing missing comma in ScriptPlugin.uplugin, adding filename to the plugin manager error so that it's actually useful

diff --git a/Engine/Plugins/ScriptPlugin/ScriptPlugin.uplugin b/Engine/Plugins/ScriptPlugin/ScriptPlugin.uplugin
--- a/Engine/Plugins/ScriptPlugin/ScriptPlugin.uplugin
+++ b/Engine/Plugins/ScriptPlugin/ScriptPlugin.uplugin
@@ -18,7 +18,7 @@
 			"Type" : "Runtime",
 			"LoadingPhase" : "PreDefault",
 			"BlacklistPlatforms" : [ "Linux" ]
-		}
+		},
 		{
 			"Name" : "ScriptEditorPlugin",
 			"Type" : "Editor",
diff --git a/Engine/Source/Runtime/Projects/Private/PluginManager.cpp b/Engine/Source/Runtime/Projects/Private/PluginManager.cpp
--- a/Engine/Source/Runtime/Projects/Private/PluginManager.cpp
+++ b/Engine/Source/Runtime/Projects/Private/PluginManager.cpp
@@ -190,7 +190,7 @@ void FPluginManager::DiscoverAllPlugins()
 					else
 					{
 						// NOTE: Even though loading of this plugin failed, we'll keep processing other plugins
-						UE_LOG(LogPluginManager, Error, TEXT("%s"), *FailureReason.ToString());
+						UE_LOG(LogPluginManager, Error, TEXT("%s (%s)"), *FailureReason.ToString(), *PluginDescriptorFilename);
 					}
 				}
 				else
