TTP#341974 CORE: Add support for external content for plugins.

diff --git a/Engine/Config/BaseEngine.ini b/Engine/Config/BaseEngine.ini
--- a/Engine/Config/BaseEngine.ini
+++ b/Engine/Config/BaseEngine.ini
@@ -1384,7 +1384,6 @@ SizeOfPermanentObjectPool=0
 AsyncIOBandwidthLimit=0
 +Paths=../../../Engine/Content
 +Paths=%GAMEDIR%Content
-+Paths=../../../Engine/Plugins/Experimental/Paper2D/Content
 CutdownPaths=%GAMEDIR%CutdownPackages
 ZeroEngineVersionWarning=True
 UseStrictEngineVersioning=True
diff --git a/Engine/Source/Runtime/Projects/Private/PluginManager.cpp b/Engine/Source/Runtime/Projects/Private/PluginManager.cpp
--- a/Engine/Source/Runtime/Projects/Private/PluginManager.cpp
+++ b/Engine/Source/Runtime/Projects/Private/PluginManager.cpp
@@ -309,6 +309,14 @@ bool FPluginManager::ConfigureEnabledPlugins()
 				ContentFolder.RootPath = FString::Printf(TEXT("/%s/"), *Plugin->Name);
 				ContentFolder.ContentPath = FPaths::GetPath(Plugin->FileName) / TEXT("Content");
 				ContentFolders.Emplace(ContentFolder);
+
+				if (auto EngineConfigFile = GConfig->Find(GEngineIni, false))
+				{
+					if (auto CoreSystemSection = EngineConfigFile->Find(TEXT("Core.System")))
+					{
+						CoreSystemSection->AddUnique("Paths", ContentFolder.ContentPath);
+					}
+				}
 			}
 		}
 
