Plugin config file usage

Merging CL 2337697

//depot/UE4-UT/...

to //depot/UE4/...

diff --git a/Engine/Source/Runtime/Core/Public/Misc/ConfigCacheIni.h b/Engine/Source/Runtime/Core/Public/Misc/ConfigCacheIni.h
--- a/Engine/Source/Runtime/Core/Public/Misc/ConfigCacheIni.h
+++ b/Engine/Source/Runtime/Core/Public/Misc/ConfigCacheIni.h
@@ -99,7 +99,7 @@ public:
 
 
 	/** Adds any properties that exist in InSourceFile that this config file is missing */
-	void AddMissingProperties( const FConfigFile& InSourceFile );
+	CORE_API void AddMissingProperties(const FConfigFile& InSourceFile);
 
 	/**
 	 * Saves only the sections in this FConfigFile its source files. All other sections in the file are left alone. The sections in this
diff --git a/Engine/Source/Runtime/Projects/Private/PluginManager.cpp b/Engine/Source/Runtime/Projects/Private/PluginManager.cpp
--- a/Engine/Source/Runtime/Projects/Private/PluginManager.cpp
+++ b/Engine/Source/Runtime/Projects/Private/PluginManager.cpp
@@ -298,25 +298,41 @@ bool FPluginManager::ConfigureEnabledPlugins()
 			}
 		}
 
-		// Build the list of content folders
 		ContentFolders.Empty();
 		for(const TSharedRef<FPluginInstance>& Plugin: AllPlugins)
 		{
-			if(Plugin->bEnabled && Plugin->Descriptor.bCanContainContent)
+			if (Plugin->bEnabled)
 			{
-				FPluginContentFolder ContentFolder;
-				ContentFolder.Name = Plugin->Name;
-				ContentFolder.RootPath = FString::Printf(TEXT("/%s/"), *Plugin->Name);
-				ContentFolder.ContentPath = FPaths::GetPath(Plugin->FileName) / TEXT("Content");
-				ContentFolders.Emplace(ContentFolder);
-
-				if (auto EngineConfigFile = GConfig->Find(GEngineIni, false))
+				// Build the list of content folders
+				if (Plugin->Descriptor.bCanContainContent)
 				{
-					if (auto CoreSystemSection = EngineConfigFile->Find(TEXT("Core.System")))
+					FPluginContentFolder ContentFolder;
+					ContentFolder.Name = Plugin->Name;
+					ContentFolder.RootPath = FString::Printf(TEXT("/%s/"), *Plugin->Name);
+					ContentFolder.ContentPath = FPaths::GetPath(Plugin->FileName) / TEXT("Content");
+					ContentFolders.Emplace(ContentFolder);
+
+					if (auto EngineConfigFile = GConfig->Find(GEngineIni, false))
 					{
-						CoreSystemSection->AddUnique("Paths", ContentFolder.ContentPath);
+						if (auto CoreSystemSection = EngineConfigFile->Find(TEXT("Core.System")))
+						{
+							CoreSystemSection->AddUnique("Paths", ContentFolder.ContentPath);
+						}
 					}
 				}
+
+				// Load Default<PluginName>.ini config file if it exists
+				FString PluginConfigDir = FPaths::GetPath(Plugin->FileName) / TEXT("Config/");
+				FConfigFile PluginConfig;
+				FConfigCacheIni::LoadExternalIniFile(PluginConfig, *Plugin->Name, *FPaths::EngineConfigDir(), *PluginConfigDir, true);
+				if (PluginConfig.Num() > 0)
+				{
+					FString PlaformName = FPlatformProperties::PlatformName();
+					FString PluginConfigFilename = FString::Printf(TEXT("%s%s/%s.ini"), *FPaths::GeneratedConfigDir(), *PlaformName, *Plugin->Name);
+					FConfigFile& NewConfigFile = GConfig->Add(PluginConfigFilename, FConfigFile());
+					NewConfigFile.AddMissingProperties(PluginConfig);
+					NewConfigFile.Write(PluginConfigFilename);
+				}
 			}
 		}
 
