Merging CL 2355060

//depot/UE4-UT/...

to //depot/UE4/...

diff --git a/Engine/Source/Runtime/Projects/Private/PluginManager.cpp b/Engine/Source/Runtime/Projects/Private/PluginManager.cpp
--- a/Engine/Source/Runtime/Projects/Private/PluginManager.cpp
+++ b/Engine/Source/Runtime/Projects/Private/PluginManager.cpp
@@ -251,6 +251,28 @@ void FPluginManager::DiscoverAllPlugins()
 	}
 }
 
+// Helper class to find all pak files.
+class FPakFileSearchVisitor : public IPlatformFile::FDirectoryVisitor
+{
+	TArray<FString>& FoundFiles;
+public:
+	FPakFileSearchVisitor(TArray<FString>& InFoundFiles)
+		: FoundFiles(InFoundFiles)
+	{}
+	virtual bool Visit(const TCHAR* FilenameOrDirectory, bool bIsDirectory)
+	{
+		if (bIsDirectory == false)
+		{
+			FString Filename(FilenameOrDirectory);
+			if (Filename.MatchesWildcard(TEXT("*.pak")))
+			{
+				FoundFiles.Add(Filename);
+			}
+		}
+		return true;
+	}
+};
+
 bool FPluginManager::ConfigureEnabledPlugins()
 {
 	if(!bHaveConfiguredEnabledPlugins)
@@ -335,13 +357,30 @@ bool FPluginManager::ConfigureEnabledPlugins()
 				}
 			}
 		}
-
-		// Mount all the plugin content folders
+		
+		// Mount all the plugin content folders and pak files
+		TArray<FString>	FoundPaks;
+		FPakFileSearchVisitor PakVisitor(FoundPaks);
+		IPlatformFile& PlatformFile = FPlatformFileManager::Get().GetPlatformFile();
 		if( ContentFolders.Num() > 0 && ensure( RegisterMountPointDelegate.IsBound() ) )
 		{
 			for(const FPluginContentFolder& ContentFolder: ContentFolders)
 			{
 				RegisterMountPointDelegate.Execute(ContentFolder.RootPath, ContentFolder.ContentPath);
+
+				// Pak files are loaded from <PluginName>/Content/Paks/<PlatformName>
+				if (FPlatformProperties::RequiresCookedData())
+				{
+					FoundPaks.Reset();
+					PlatformFile.IterateDirectoryRecursively(*(ContentFolder.ContentPath / TEXT("Paks") / FPlatformProperties::PlatformName()), PakVisitor);
+					for (const auto& PakPath : FoundPaks)
+					{
+						if (FCoreDelegates::OnMountPak.IsBound())
+						{
+							FCoreDelegates::OnMountPak.Execute(PakPath, 0);
+						}
+					}
+				}
 			}
 		}
 	}
