Getting rid of unnecessary loading of UE4Game module.

diff --git a/Engine/Source/Runtime/Projects/Private/ProjectManager.cpp b/Engine/Source/Runtime/Projects/Private/ProjectManager.cpp
--- a/Engine/Source/Runtime/Projects/Private/ProjectManager.cpp
+++ b/Engine/Source/Runtime/Projects/Private/ProjectManager.cpp
@@ -93,12 +93,6 @@ bool FProjectManager::LoadModulesForProject( const ELoadingPhase::Type LoadingPh
 		}
 	}
 
-	if (LoadingPhase == ELoadingPhase::PreDefault && (!CurrentProject.IsValid() || CurrentProject->Modules.Num() == 0))
-	{
-		// If this is not a code-project load default game project module.
-		FModuleManager::Get().LoadModuleChecked("UE4Game");
-	}
-
 	return bSuccess;
 }
 
