Better explanation for what FProjectDescriptor::EngineAssociation does.

diff --git a/Engine/Source/Runtime/Projects/Public/ProjectDescriptor.h b/Engine/Source/Runtime/Projects/Public/ProjectDescriptor.h
--- a/Engine/Source/Runtime/Projects/Public/ProjectDescriptor.h
+++ b/Engine/Source/Runtime/Projects/Public/ProjectDescriptor.h
@@ -32,7 +32,35 @@ struct PROJECTS_API FProjectDescriptor
 	/** Descriptor version number. */
 	EProjectDescriptorVersion::Type FileVersion;
 
-	/** The associated engine. Used by UnrealVersionSelector to open projects with the correct editor. Set this value using IDesktopPlatform::SetEngineIdentifierForProject. */
+	/** 
+	 * The engine to open this project with. Set this value using IDesktopPlatform::SetEngineIdentifierForProject to ensure that
+	 * the most portable value for this field is used.
+	 *
+	 * This field allows us to open the right version of the engine when you double-click on a .uproject file, and to detect when you 
+	 * open a project with a different version of the editor and need the upgrade/downgrade UI flow. The normal engine 
+	 * version doesn't work for those purposes, because you can have multiple 4.x branches in various states on one machine.
+	 *
+	 * For Rocket users, this field gets set to something stable like "4.7" or "4.8", so you can swap projects and game binaries 
+	 * between users, and it'll automatically work on any platform or machine regardless of where the engine is installed. You 
+	 * can only have one binary release of each major engine version installed at once.
+	 * 
+	 * For Perforce or Git users that branch the engine along with their games, this field is left blank. You can sync the repository 
+	 * down on any platform and machine, and it can figure out which engine a project should use by looking up the directory 
+	 * hierarchy until it finds one.
+	 * 
+	 * For other cases, where you have a source build of the engine but are working with a foreign project, we use a random identifier 
+	 * for each local engine installation and use the registry to map it back to the engine directory. All bets are off as to which
+	 * engine you should use to open it on a different machine, and using a random GUID ensures that every new machine triggers the
+	 * engine selection UI when you open or attempt to generate project files for it.
+	 * 
+	 * For users which mount the engine through a Git submodule (where the engine is in a subdirectory of the project), this field 
+	 * can be manually edited to be a relative path.
+	 *
+	 * @see IDesktopPlatform::GetEngineIdentifierForProject
+	 * @see IDesktopPlatform::SetEngineIdentifierForProject
+	 * @see IDesktopPlatform::GetEngineRootDirFromIdentifier
+	 * @see IDesktopPlatform::GetEngineIdentifierFromRootDir
+	 */
 	FString EngineAssociation;
 
 	/** Category to show under the project browser */
