Add IPlugin interface, rather than exposing lots of raw data for plugins.

diff --git a/Engine/Source/Runtime/Projects/Private/PluginManager.cpp b/Engine/Source/Runtime/Projects/Private/PluginManager.cpp
--- a/Engine/Source/Runtime/Projects/Private/PluginManager.cpp
+++ b/Engine/Source/Runtime/Projects/Private/PluginManager.cpp
@@ -15,7 +15,7 @@ namespace PluginSystemDefs
 
 }
 
-FPluginInstance::FPluginInstance(const FString& InFileName, const FPluginDescriptor& InDescriptor, EPluginLoadedFrom InLoadedFrom)
+FPlugin::FPlugin(const FString& InFileName, const FPluginDescriptor& InDescriptor, EPluginLoadedFrom InLoadedFrom)
 	: Name(FPaths::GetBaseFilename(InFileName))
 	, FileName(InFileName)
 	, Descriptor(InDescriptor)
@@ -24,7 +24,35 @@ FPluginInstance::FPluginInstance(const FString& InFileName, const FPluginDescrip
 {
 }
 
+FString FPlugin::GetName() const
+{
+	return Name;
+}
+
+FString FPlugin::GetBaseDir() const
+{
+	return FPaths::GetPath(FileName);
+}
+
+FString FPlugin::GetContentDir() const
+{
+	return FPaths::GetPath(FileName) / TEXT("Content");
+}
+
+FString FPlugin::GetMountedAssetPath() const
+{
+	return FString::Printf(TEXT("/%s/"), *Name);
+}
+
+bool FPlugin::IsEnabled() const
+{
+	return bEnabled;
+}
 
+bool FPlugin::CanContainContent() const
+{
+	return Descriptor.bCanContainContent;
+}
 
 
 
@@ -63,7 +91,7 @@ void FPluginManager::DiscoverAllPlugins()
 		 * @param	LoadedFrom			Where we're loading these plugins from (game, engine, etc)
 		 * @param	Plugins				The array to be filled in with new plugins including descriptors
 		 */
-		static void FindPluginsRecursively( const FString& PluginsDirectory, const EPluginLoadedFrom LoadedFrom, TArray< TSharedRef<FPluginInstance> >& Plugins )
+		static void FindPluginsRecursively( const FString& PluginsDirectory, const EPluginLoadedFrom LoadedFrom, TArray< TSharedRef<FPlugin> >& Plugins )
 		{
 			// NOTE: The logic in this function generally matches that of the C# code for FindPluginsRecursively
 			//       in UnrealBuildTool.  These routines should be kept in sync.
@@ -170,7 +198,7 @@ void FPluginManager::DiscoverAllPlugins()
 					FText FailureReason;
 					if(Descriptor.Load(PluginDescriptorFilename, FailureReason))
 					{
-						TSharedRef< FPluginInstance > NewPlugin = MakeShareable( new FPluginInstance(PluginDescriptorFilename, Descriptor, LoadedFrom) );
+						TSharedRef< FPlugin > NewPlugin = MakeShareable( new FPlugin(PluginDescriptorFilename, Descriptor, LoadedFrom) );
 
 						if (FPaths::IsProjectFilePathSet())
 						{
@@ -208,7 +236,7 @@ void FPluginManager::DiscoverAllPlugins()
 		 * @param	LoadedFrom			Where we're loading these plugins from (game, engine, etc)
 		 * @param	Plugins				The array to be filled in with loaded plugin descriptors
 		 */
-		static void FindPluginsIn( const FString& PluginsDirectory, const EPluginLoadedFrom LoadedFrom, TArray< TSharedRef<FPluginInstance> >& Plugins )
+		static void FindPluginsIn( const FString& PluginsDirectory, const EPluginLoadedFrom LoadedFrom, TArray< TSharedRef<FPlugin> >& Plugins )
 		{
 			// Make sure the directory even exists
 			if( FPlatformFileManager::Get().GetPlatformFile().DirectoryExists( *PluginsDirectory ) )
@@ -220,7 +248,7 @@ void FPluginManager::DiscoverAllPlugins()
 
 
 	{
-		TArray< TSharedRef<FPluginInstance> > Plugins;
+		TArray< TSharedRef<FPlugin> > Plugins;
 
 #if (WITH_ENGINE && !IS_PROGRAM) || WITH_PLUGIN_SUPPORT
 		// Find "built-in" plugins.  That is, plugins situated right within the Engine directory.
@@ -239,7 +267,7 @@ void FPluginManager::DiscoverAllPlugins()
 		AllPlugins = Plugins;
 		for( auto PluginIt( Plugins.CreateConstIterator() ); PluginIt; ++PluginIt )
 		{
-			const TSharedRef<FPluginInstance>& Plugin = *PluginIt;
+			const TSharedRef<FPlugin>& Plugin = *PluginIt;
 
 			// Add the plugin binaries directory
 			const FString PluginBinariesPath = FPaths::Combine(*FPaths::GetPath(Plugin->FileName), TEXT("Binaries"), FPlatformProcess::GetBinariesSubdirectory());
@@ -309,7 +337,7 @@ bool FPluginManager::ConfigureEnabledPlugins()
 		AllEnabledPlugins.Append( MoveTemp(EnabledPluginNames) );
 
 		// Enable all the plugins by name
-		for( const TSharedRef< FPluginInstance > Plugin : AllPlugins )
+		for( const TSharedRef< FPlugin > Plugin : AllPlugins )
 		{
 			if ( AllEnabledPlugins.Contains(Plugin->Name) )
 			{
@@ -318,7 +346,7 @@ bool FPluginManager::ConfigureEnabledPlugins()
 		}
 
 		ContentFolders.Empty();
-		for(const TSharedRef<FPluginInstance>& Plugin: AllPlugins)
+		for(const TSharedRef<FPlugin>& Plugin: AllPlugins)
 		{
 			if (Plugin->bEnabled)
 			{
@@ -326,9 +354,9 @@ bool FPluginManager::ConfigureEnabledPlugins()
 				if (Plugin->Descriptor.bCanContainContent)
 				{
 					FPluginContentFolder ContentFolder;
-					ContentFolder.Name = Plugin->Name;
-					ContentFolder.RootPath = FString::Printf(TEXT("/%s/"), *Plugin->Name);
-					ContentFolder.ContentPath = FPaths::GetPath(Plugin->FileName) / TEXT("Content");
+					ContentFolder.Name = Plugin->GetName();
+					ContentFolder.RootPath = Plugin->GetMountedAssetPath();
+					ContentFolder.ContentPath = Plugin->GetContentDir();
 					ContentFolders.Emplace(ContentFolder);
 
 					if (auto EngineConfigFile = GConfig->Find(GEngineIni, false))
@@ -384,10 +412,10 @@ bool FPluginManager::ConfigureEnabledPlugins()
 	return bHaveAllRequiredPlugins;
 }
 
-TSharedPtr<FPluginInstance> FPluginManager::FindPluginInstance(const FString& Name)
+TSharedPtr<FPlugin> FPluginManager::FindPluginInstance(const FString& Name)
 {
-	TSharedPtr<FPluginInstance> Result;
-	for(const TSharedRef<FPluginInstance>& Instance : AllPlugins)
+	TSharedPtr<FPlugin> Result;
+	for(const TSharedRef<FPlugin>& Instance : AllPlugins)
 	{
 		if(Instance->Name == Name)
 		{
@@ -409,7 +437,7 @@ bool FPluginManager::LoadModulesForEnabledPlugins( const ELoadingPhase::Type Loa
 	FScopedSlowTask SlowTask(AllPlugins.Num());
 
 	// Load plugins!
-	for( const TSharedRef< FPluginInstance > Plugin : AllPlugins )
+	for( const TSharedRef< FPlugin > Plugin : AllPlugins )
 	{
 		SlowTask.EnterProgressFrame(1);
 
@@ -484,9 +512,9 @@ bool FPluginManager::CheckModuleCompatibility(TArray<FString>& OutIncompatibleMo
 	}
 
 	bool bResult = true;
-	for (TArray< TSharedRef< FPluginInstance > >::TConstIterator Iter(AllPlugins); Iter; ++Iter)
+	for (TArray< TSharedRef< FPlugin > >::TConstIterator Iter(AllPlugins); Iter; ++Iter)
 	{
-		const TSharedRef< FPluginInstance > &Plugin = *Iter;
+		const TSharedRef< FPlugin > &Plugin = *Iter;
 		if (Plugin->bEnabled && !FModuleDescriptor::CheckModuleCompatbility(Plugin->Descriptor.Modules, Plugin->LoadedFrom == EPluginLoadedFrom::GameProject, OutIncompatibleModules))
 		{
 			bResult = false;
@@ -506,13 +534,27 @@ IPluginManager& IPluginManager::Get()
 	return *PluginManager;
 }
 
+FPlugin* FPluginManager::FindPlugin(const FString& Name)
+{
+	FPlugin* Plugin = nullptr;
+	for(TSharedRef<FPlugin>& PossiblePlugin : AllPlugins)
+	{
+		if(PossiblePlugin->Name == Name)
+		{
+			Plugin = &(PossiblePlugin.Get());
+			break;
+		}
+	}
+	return Plugin;
+}
+
 TArray< FPluginStatus > FPluginManager::QueryStatusForAllPlugins() const
 {
 	TArray< FPluginStatus > PluginStatuses;
 
 	for( auto PluginIt( AllPlugins.CreateConstIterator() ); PluginIt; ++PluginIt )
 	{
-		const TSharedRef< FPluginInstance >& Plugin = *PluginIt;
+		const TSharedRef< FPlugin >& Plugin = *PluginIt;
 		
 		FPluginStatus PluginStatus;
 		PluginStatus.Name = Plugin->Name;
diff --git a/Engine/Source/Runtime/Projects/Private/PluginManager.h b/Engine/Source/Runtime/Projects/Private/PluginManager.h
--- a/Engine/Source/Runtime/Projects/Private/PluginManager.h
+++ b/Engine/Source/Runtime/Projects/Private/PluginManager.h
@@ -7,7 +7,7 @@
 /**
  * Instance of a plugin in memory
  */
-class FPluginInstance
+class FPlugin : public IPlugin
 {
 public:
 	/** The name of the plugin */
@@ -28,7 +28,15 @@ public:
 	/**
 	 * FPlugin constructor
 	 */
-	FPluginInstance(const FString &FileName, const FPluginDescriptor& InDescriptor, EPluginLoadedFrom InLoadedFrom);
+	FPlugin(const FString &FileName, const FPluginDescriptor& InDescriptor, EPluginLoadedFrom InLoadedFrom);
+
+	/* IPluginInfo interface */
+	virtual FString GetName() const override;
+	virtual FString GetBaseDir() const override;
+	virtual FString GetContentDir() const override;
+	virtual FString GetMountedAssetPath() const override;
+	virtual bool IsEnabled() const override;
+	virtual bool CanContainContent() const override;
 };
 
 /**
@@ -48,6 +56,7 @@ public:
 	virtual void SetRegisterMountPointDelegate( const FRegisterMountPointDelegate& Delegate ) override;
 	virtual bool AreRequiredPluginsAvailable() override;
 	virtual bool CheckModuleCompatibility( TArray<FString>& OutIncompatibleModules ) override;
+	virtual FPlugin* FindPlugin(const FString& Name) override;
 	virtual TArray< FPluginStatus > QueryStatusForAllPlugins() const override;
 	virtual const TArray< FPluginContentFolder >& GetPluginContentFolders() const override;
 
@@ -61,11 +70,11 @@ private:
 	bool ConfigureEnabledPlugins();
 
 	/** Gets the instance of a given plugin */
-	TSharedPtr<FPluginInstance> FindPluginInstance(const FString& Name);
+	TSharedPtr<FPlugin> FindPluginInstance(const FString& Name);
 
 private:
 	/** All of the plugins that we know about */
-	TArray< TSharedRef< FPluginInstance > > AllPlugins;
+	TArray< TSharedRef< FPlugin > > AllPlugins;
 
 	/** All the plugin content folders */
 	TArray<FPluginContentFolder> ContentFolders;
diff --git a/Engine/Source/Runtime/Projects/Public/Interfaces/IPluginManager.h b/Engine/Source/Runtime/Projects/Public/Interfaces/IPluginManager.h
--- a/Engine/Source/Runtime/Projects/Public/Interfaces/IPluginManager.h
+++ b/Engine/Source/Runtime/Projects/Public/Interfaces/IPluginManager.h
@@ -55,6 +55,54 @@ struct FPluginContentFolder
 	FString ContentPath;
 };
 
+/**
+ * Information about an enabled plugin.
+ */
+class IPlugin
+{
+public:
+	/**
+	 * Gets the plugin name.
+	 *
+	 * @return Name of the plugin.
+	 */
+	virtual FString GetName() const = 0;
+
+	/**
+	 * Get a path to the plugin's directory.
+	 *
+	 * @return Path to the plugin's base directory.
+	 */
+	virtual FString GetBaseDir() const = 0;
+
+	/**
+	 * Get a path to the plugin's content directory.
+	 *
+	 * @return Path to the plugin's content directory.
+	 */
+	virtual FString GetContentDir() const = 0;
+
+	/**
+	 * Get the virtual root path for assets.
+	 *
+	 * @return The mounted root path for assets in this plugin's content folder; typically /PluginName/.
+	 */
+	virtual FString GetMountedAssetPath() const = 0;
+
+	/**
+	 * Determines if the plugin is enabled.
+	 *
+	 * @return True if the plugin is currently enabled.
+	 */
+	virtual bool IsEnabled() const = 0;
+
+	/**
+	 * Determines if the plugin can contain content.
+	 *
+	 * @return True if the plugin can contain content.
+	 */
+	virtual bool CanContainContent() const = 0;
+};
 
 /**
  * PluginManager manages available code and content extensions (both loaded and not loaded).
@@ -97,6 +145,13 @@ public:
 	 */
 	virtual bool CheckModuleCompatibility( TArray<FString>& OutIncompatibleModules ) = 0;
 
+	/**
+	 * Finds information for an enabled plugin.
+	 *
+	 * @return	 Pointer to the plugin's information, or nullptr.
+	 */
+	virtual IPlugin* FindPlugin(const FString& Name) = 0;
+
 	/**
 	 * Gets status about all currently known plug-ins.
 	 *
