Accessor to get an array of all the enabled plugins.

diff --git a/Engine/Source/Runtime/Projects/Private/PluginManager.cpp b/Engine/Source/Runtime/Projects/Private/PluginManager.cpp
--- a/Engine/Source/Runtime/Projects/Private/PluginManager.cpp
+++ b/Engine/Source/Runtime/Projects/Private/PluginManager.cpp
@@ -548,6 +548,19 @@ FPlugin* FPluginManager::FindPlugin(const FString& Name)
 	return Plugin;
 }
 
+TArray<IPlugin*> FPluginManager::GetEnabledPlugins()
+{
+	TArray<IPlugin*> Plugins;
+	for(TSharedRef<FPlugin>& PossiblePlugin : AllPlugins)
+	{
+		if(PossiblePlugin->bEnabled)
+		{
+			Plugins.Add(&(PossiblePlugin.Get()));
+		}
+	}
+	return Plugins;
+}
+
 TArray< FPluginStatus > FPluginManager::QueryStatusForAllPlugins() const
 {
 	TArray< FPluginStatus > PluginStatuses;
diff --git a/Engine/Source/Runtime/Projects/Private/PluginManager.h b/Engine/Source/Runtime/Projects/Private/PluginManager.h
--- a/Engine/Source/Runtime/Projects/Private/PluginManager.h
+++ b/Engine/Source/Runtime/Projects/Private/PluginManager.h
@@ -57,6 +57,7 @@ public:
 	virtual bool AreRequiredPluginsAvailable() override;
 	virtual bool CheckModuleCompatibility( TArray<FString>& OutIncompatibleModules ) override;
 	virtual FPlugin* FindPlugin(const FString& Name) override;
+	virtual TArray<IPlugin*> GetEnabledPlugins() override;
 	virtual TArray< FPluginStatus > QueryStatusForAllPlugins() const override;
 	virtual const TArray< FPluginContentFolder >& GetPluginContentFolders() const override;
 
diff --git a/Engine/Source/Runtime/Projects/Public/Interfaces/IPluginManager.h b/Engine/Source/Runtime/Projects/Public/Interfaces/IPluginManager.h
--- a/Engine/Source/Runtime/Projects/Public/Interfaces/IPluginManager.h
+++ b/Engine/Source/Runtime/Projects/Public/Interfaces/IPluginManager.h
@@ -152,6 +152,13 @@ public:
 	 */
 	virtual IPlugin* FindPlugin(const FString& Name) = 0;
 
+	/**
+	 * Gets an array of all the enabled plugins.
+	 *
+	 * @return	Array of the enabled plugins.
+	 */
+	virtual TArray<IPlugin*> GetEnabledPlugins() = 0;
+
 	/**
 	 * Gets status about all currently known plug-ins.
 	 *
