Use the exposed IPlugin interface to get plugin content folders rather than querying it separately.

diff --git a/Engine/Source/Editor/ContentBrowser/Private/ContentBrowserUtils.cpp b/Engine/Source/Editor/ContentBrowser/Private/ContentBrowserUtils.cpp
--- a/Engine/Source/Editor/ContentBrowser/Private/ContentBrowserUtils.cpp
+++ b/Engine/Source/Editor/ContentBrowser/Private/ContentBrowserUtils.cpp
@@ -1097,11 +1097,14 @@ bool ContentBrowserUtils::IsDevelopersFolder( const FString& InPath )
 bool ContentBrowserUtils::IsPluginFolder( const FString& InPath )
 {
 	FString PathWithSlash = InPath / TEXT("");
-	for(const FPluginContentFolder& ContentFolder: IPluginManager::Get().GetPluginContentFolders())
+	for(const IPlugin* Plugin: IPluginManager::Get().GetEnabledPlugins())
 	{
-		if(PathWithSlash.StartsWith(ContentFolder.RootPath) || InPath == ContentFolder.Name)
+		if(Plugin->CanContainContent())
 		{
-			return true;
+			if(PathWithSlash.StartsWith(Plugin->GetMountedAssetPath()) || InPath == Plugin->GetName())
+			{
+				return true;
+			}
 		}
 	}
 	return false;
diff --git a/Engine/Source/Editor/ContentBrowser/Private/SAssetView.cpp b/Engine/Source/Editor/ContentBrowser/Private/SAssetView.cpp
--- a/Engine/Source/Editor/ContentBrowser/Private/SAssetView.cpp
+++ b/Engine/Source/Editor/ContentBrowser/Private/SAssetView.cpp
@@ -2327,21 +2327,18 @@ TSharedRef<SWidget> SAssetView::GetViewButtonContent()
 			EUserInterfaceActionType::ToggleButton
 			);
 
-		if(IPluginManager::Get().GetPluginContentFolders().Num() > 0)
-		{
-			MenuBuilder.AddMenuEntry(
-				LOCTEXT("ShowPluginFolderOption", "Show Plugin Content"),
-				LOCTEXT("ShowPluginFolderOptionToolTip", "Shows plugin content in the view."),
-				FSlateIcon(),
-				FUIAction(
-					FExecuteAction::CreateSP( this, &SAssetView::ToggleShowPluginFolders ),
-					FCanExecuteAction(),
-					FIsActionChecked::CreateSP( this, &SAssetView::IsShowingPluginFolders )
-				),
-				NAME_None,
-				EUserInterfaceActionType::ToggleButton
-				);
-		}
+		MenuBuilder.AddMenuEntry(
+			LOCTEXT("ShowPluginFolderOption", "Show Plugin Content"),
+			LOCTEXT("ShowPluginFolderOptionToolTip", "Shows plugin content in the view."),
+			FSlateIcon(),
+			FUIAction(
+				FExecuteAction::CreateSP( this, &SAssetView::ToggleShowPluginFolders ),
+				FCanExecuteAction(),
+				FIsActionChecked::CreateSP( this, &SAssetView::IsShowingPluginFolders )
+			),
+			NAME_None,
+			EUserInterfaceActionType::ToggleButton
+			);
 
 		MenuBuilder.AddMenuEntry(
 			LOCTEXT("ShowEngineFolderOption", "Show Engine Content"),
diff --git a/Engine/Source/Runtime/Projects/Private/PluginManager.cpp b/Engine/Source/Runtime/Projects/Private/PluginManager.cpp
--- a/Engine/Source/Runtime/Projects/Private/PluginManager.cpp
+++ b/Engine/Source/Runtime/Projects/Private/PluginManager.cpp
@@ -345,7 +345,6 @@ bool FPluginManager::ConfigureEnabledPlugins()
 			}
 		}
 
-		ContentFolders.Empty();
 		for(const TSharedRef<FPlugin>& Plugin: AllPlugins)
 		{
 			if (Plugin->bEnabled)
@@ -353,17 +352,11 @@ bool FPluginManager::ConfigureEnabledPlugins()
 				// Build the list of content folders
 				if (Plugin->Descriptor.bCanContainContent)
 				{
-					FPluginContentFolder ContentFolder;
-					ContentFolder.Name = Plugin->GetName();
-					ContentFolder.RootPath = Plugin->GetMountedAssetPath();
-					ContentFolder.ContentPath = Plugin->GetContentDir();
-					ContentFolders.Emplace(ContentFolder);
-
 					if (auto EngineConfigFile = GConfig->Find(GEngineIni, false))
 					{
 						if (auto CoreSystemSection = EngineConfigFile->Find(TEXT("Core.System")))
 						{
-							CoreSystemSection->AddUnique("Paths", ContentFolder.ContentPath);
+							CoreSystemSection->AddUnique("Paths", Plugin->GetContentDir());
 						}
 					}
 				}
@@ -387,22 +380,26 @@ bool FPluginManager::ConfigureEnabledPlugins()
 		TArray<FString>	FoundPaks;
 		FPakFileSearchVisitor PakVisitor(FoundPaks);
 		IPlatformFile& PlatformFile = FPlatformFileManager::Get().GetPlatformFile();
-		if( ContentFolders.Num() > 0 && ensure( RegisterMountPointDelegate.IsBound() ) )
+		if( ensure( RegisterMountPointDelegate.IsBound() ) )
 		{
-			for(const FPluginContentFolder& ContentFolder: ContentFolders)
+			for(IPlugin* Plugin: GetEnabledPlugins())
 			{
-				RegisterMountPointDelegate.Execute(ContentFolder.RootPath, ContentFolder.ContentPath);
-
-				// Pak files are loaded from <PluginName>/Content/Paks/<PlatformName>
-				if (FPlatformProperties::RequiresCookedData())
+				if(Plugin->CanContainContent())
 				{
-					FoundPaks.Reset();
-					PlatformFile.IterateDirectoryRecursively(*(ContentFolder.ContentPath / TEXT("Paks") / FPlatformProperties::PlatformName()), PakVisitor);
-					for (const auto& PakPath : FoundPaks)
+					FString ContentDir = Plugin->GetContentDir();
+					RegisterMountPointDelegate.Execute(Plugin->GetMountedAssetPath(), ContentDir);
+
+					// Pak files are loaded from <PluginName>/Content/Paks/<PlatformName>
+					if (FPlatformProperties::RequiresCookedData())
 					{
-						if (FCoreDelegates::OnMountPak.IsBound())
+						FoundPaks.Reset();
+						PlatformFile.IterateDirectoryRecursively(*(ContentDir / TEXT("Paks") / FPlatformProperties::PlatformName()), PakVisitor);
+						for (const auto& PakPath : FoundPaks)
 						{
-							FCoreDelegates::OnMountPak.Execute(PakPath, 0);
+							if (FCoreDelegates::OnMountPak.IsBound())
+							{
+								FCoreDelegates::OnMountPak.Execute(PakPath, 0);
+							}
 						}
 					}
 				}
@@ -582,9 +579,4 @@ TArray< FPluginStatus > FPluginManager::QueryStatusForAllPlugins() const
 	return PluginStatuses;
 }
 
-const TArray<FPluginContentFolder>& FPluginManager::GetPluginContentFolders() const
-{
-	return ContentFolders;
-}
-
 #undef LOCTEXT_NAMESPACE
diff --git a/Engine/Source/Runtime/Projects/Private/PluginManager.h b/Engine/Source/Runtime/Projects/Private/PluginManager.h
--- a/Engine/Source/Runtime/Projects/Private/PluginManager.h
+++ b/Engine/Source/Runtime/Projects/Private/PluginManager.h
@@ -59,7 +59,6 @@ public:
 	virtual FPlugin* FindPlugin(const FString& Name) override;
 	virtual TArray<IPlugin*> GetEnabledPlugins() override;
 	virtual TArray< FPluginStatus > QueryStatusForAllPlugins() const override;
-	virtual const TArray< FPluginContentFolder >& GetPluginContentFolders() const override;
 
 private:
 
@@ -77,9 +76,6 @@ private:
 	/** All of the plugins that we know about */
 	TArray< TSharedRef< FPlugin > > AllPlugins;
 
-	/** All the plugin content folders */
-	TArray<FPluginContentFolder> ContentFolders;
-
 	/** Delegate for mounting content paths.  Bound by FPackageName code in CoreUObject, so that we can access
 	    content path mounting functionality from Core. */
 	FRegisterMountPointDelegate RegisterMountPointDelegate;
diff --git a/Engine/Source/Runtime/Projects/Public/Interfaces/IPluginManager.h b/Engine/Source/Runtime/Projects/Public/Interfaces/IPluginManager.h
--- a/Engine/Source/Runtime/Projects/Public/Interfaces/IPluginManager.h
+++ b/Engine/Source/Runtime/Projects/Public/Interfaces/IPluginManager.h
@@ -40,21 +40,6 @@ struct FPluginStatus
 };
 
 
-/**
- * Structure holding information about a plug-in content folder.
- */
-struct FPluginContentFolder
-{
-	/** Name of the plug-in */
-	FString Name;
-
-	/** Virtual root path for asset paths */
-	FString RootPath;
-
-	/** Content path on disk */
-	FString ContentPath;
-};
-
 /**
  * Information about an enabled plugin.
  */
@@ -166,13 +151,6 @@ public:
 	 */
 	virtual TArray<FPluginStatus> QueryStatusForAllPlugins() const = 0;
 
-	/**
-	 * Gets a list of plug-in content folders.
-	 *
-	 * @return	 Array of plug-in content folders.
-	 */
-	virtual const TArray<FPluginContentFolder>& GetPluginContentFolders() const = 0;
-
 public:
 
 	/**
