Add a virtual destructor to FPlugin.

diff --git a/Engine/Source/Runtime/Projects/Private/PluginManager.cpp b/Engine/Source/Runtime/Projects/Private/PluginManager.cpp
--- a/Engine/Source/Runtime/Projects/Private/PluginManager.cpp
+++ b/Engine/Source/Runtime/Projects/Private/PluginManager.cpp
@@ -24,6 +24,10 @@ FPlugin::FPlugin(const FString& InFileName, const FPluginDescriptor& InDescripto
 {
 }
 
+FPlugin::~FPlugin()
+{
+}
+
 FString FPlugin::GetName() const
 {
 	return Name;
diff --git a/Engine/Source/Runtime/Projects/Private/PluginManager.h b/Engine/Source/Runtime/Projects/Private/PluginManager.h
--- a/Engine/Source/Runtime/Projects/Private/PluginManager.h
+++ b/Engine/Source/Runtime/Projects/Private/PluginManager.h
@@ -30,6 +30,11 @@ public:
 	 */
 	FPlugin(const FString &FileName, const FPluginDescriptor& InDescriptor, EPluginLoadedFrom InLoadedFrom);
 
+	/**
+	 * Destructor.
+	 */
+	virtual ~FPlugin();
+
 	/* IPluginInfo interface */
 	virtual FString GetName() const override;
 	virtual FString GetBaseDir() const override;
