Remove ensure against CoreUObject mount delegate existing unless there are plugins with content folders.  Not entirely sure why this is necessary yet, but Mac DDC seems to be failing where it wasn't before.

diff --git a/Engine/Source/Runtime/Projects/Private/PluginManager.cpp b/Engine/Source/Runtime/Projects/Private/PluginManager.cpp
--- a/Engine/Source/Runtime/Projects/Private/PluginManager.cpp
+++ b/Engine/Source/Runtime/Projects/Private/PluginManager.cpp
@@ -409,26 +409,23 @@ bool FPluginManager::ConfigureEnabledPlugins()
 		TArray<FString>	FoundPaks;
 		FPakFileSearchVisitor PakVisitor(FoundPaks);
 		IPlatformFile& PlatformFile = FPlatformFileManager::Get().GetPlatformFile();
-		if( ensure( RegisterMountPointDelegate.IsBound() ) )
+		for(TSharedRef<IPlugin> Plugin: GetEnabledPlugins())
 		{
-			for(TSharedRef<IPlugin> Plugin: GetEnabledPlugins())
+			if(Plugin->CanContainContent() && ensure(RegisterMountPointDelegate.IsBound()))
 			{
-				if(Plugin->CanContainContent())
-				{
-					FString ContentDir = Plugin->GetContentDir();
-					RegisterMountPointDelegate.Execute(Plugin->GetMountedAssetPath(), ContentDir);
+				FString ContentDir = Plugin->GetContentDir();
+				RegisterMountPointDelegate.Execute(Plugin->GetMountedAssetPath(), ContentDir);
 
-					// Pak files are loaded from <PluginName>/Content/Paks/<PlatformName>
-					if (FPlatformProperties::RequiresCookedData())
+				// Pak files are loaded from <PluginName>/Content/Paks/<PlatformName>
+				if (FPlatformProperties::RequiresCookedData())
+				{
+					FoundPaks.Reset();
+					PlatformFile.IterateDirectoryRecursively(*(ContentDir / TEXT("Paks") / FPlatformProperties::PlatformName()), PakVisitor);
+					for (const auto& PakPath : FoundPaks)
 					{
-						FoundPaks.Reset();
-						PlatformFile.IterateDirectoryRecursively(*(ContentDir / TEXT("Paks") / FPlatformProperties::PlatformName()), PakVisitor);
-						for (const auto& PakPath : FoundPaks)
+						if (FCoreDelegates::OnMountPak.IsBound())
 						{
-							if (FCoreDelegates::OnMountPak.IsBound())
-							{
-								FCoreDelegates::OnMountPak.Execute(PakPath, 0);
-							}
+							FCoreDelegates::OnMountPak.Execute(PakPath, 0);
 						}
 					}
 				}
