Fix initialization order warnings for FPluginReferenceDescriptor, and reorder constructor parameters to match the more logical order.

diff --git a/Engine/Source/Runtime/Projects/Private/PluginDescriptor.cpp b/Engine/Source/Runtime/Projects/Private/PluginDescriptor.cpp
--- a/Engine/Source/Runtime/Projects/Private/PluginDescriptor.cpp
+++ b/Engine/Source/Runtime/Projects/Private/PluginDescriptor.cpp
@@ -153,10 +153,10 @@ FString FPluginDescriptor::ToString() const
 
 //////////////////////////////////////////////////////////////////////////////////////////////////////////////////
 
-FPluginReferenceDescriptor::FPluginReferenceDescriptor( const FString& InName, const FString& InMarketplaceURL, bool bInEnabled )
+FPluginReferenceDescriptor::FPluginReferenceDescriptor( const FString& InName, bool bInEnabled, const FString& InMarketplaceURL )
 	: Name(InName)
-	, MarketplaceURL(InMarketplaceURL)
 	, bEnabled(bInEnabled)
+	, MarketplaceURL(InMarketplaceURL)
 { }
 
 
diff --git a/Engine/Source/Runtime/Projects/Private/ProjectManager.cpp b/Engine/Source/Runtime/Projects/Private/ProjectManager.cpp
--- a/Engine/Source/Runtime/Projects/Private/ProjectManager.cpp
+++ b/Engine/Source/Runtime/Projects/Private/ProjectManager.cpp
@@ -290,7 +290,7 @@ bool FProjectManager::SetPluginEnabled(const FString& PluginName, bool bEnabled,
 	{
 		if(PluginRefIdx == CurrentProject->Plugins.Num())
 		{
-			PluginRefIdx = CurrentProject->Plugins.Add(FPluginReferenceDescriptor(PluginName, MarketplaceURL, bEnabled));
+			PluginRefIdx = CurrentProject->Plugins.Add(FPluginReferenceDescriptor(PluginName, bEnabled, MarketplaceURL));
 			break;
 		}
 		else if(CurrentProject->Plugins[PluginRefIdx].Name == PluginName)
diff --git a/Engine/Source/Runtime/Projects/Public/PluginDescriptor.h b/Engine/Source/Runtime/Projects/Public/PluginDescriptor.h
--- a/Engine/Source/Runtime/Projects/Public/PluginDescriptor.h
+++ b/Engine/Source/Runtime/Projects/Public/PluginDescriptor.h
@@ -116,7 +116,7 @@ struct PROJECTS_API FPluginReferenceDescriptor
 	TArray<FString> BlacklistPlatforms;
 
 	/** Constructor */
-	FPluginReferenceDescriptor(const FString& InName = TEXT(""), const FString& InMarketplaceURL = TEXT(""), bool bInEnabled = false);
+	FPluginReferenceDescriptor(const FString& InName = TEXT(""), bool bInEnabled = false, const FString& InMarketplaceURL = TEXT(""));
 
 	/** Determines whether the plugin is enabled for the given platform */
 	bool IsEnabledForPlatform(const FString& Platform) const;
