Automatically refresh the plugins list when a plugin is added or modified externally. Also simplify the logic for scanning for plugins.

diff --git a/Engine/Plugins/Editor/PluginBrowser/Source/PluginBrowser/PluginBrowser.Build.cs b/Engine/Plugins/Editor/PluginBrowser/Source/PluginBrowser/PluginBrowser.Build.cs
--- a/Engine/Plugins/Editor/PluginBrowser/Source/PluginBrowser/PluginBrowser.Build.cs
+++ b/Engine/Plugins/Editor/PluginBrowser/Source/PluginBrowser/PluginBrowser.Build.cs
@@ -23,7 +23,8 @@ namespace UnrealBuildTool.Rules
 					"Projects",
 					"UnrealEd",
 					"PropertyEditor",
-					"SharedSettingsWidgets"
+					"SharedSettingsWidgets",
+					"DirectoryWatcher"
 				}
 			);
 
diff --git a/Engine/Plugins/Editor/PluginBrowser/Source/PluginBrowser/Private/SPluginBrowser.cpp b/Engine/Plugins/Editor/PluginBrowser/Source/PluginBrowser/Private/SPluginBrowser.cpp
--- a/Engine/Plugins/Editor/PluginBrowser/Source/PluginBrowser/Private/SPluginBrowser.cpp
+++ b/Engine/Plugins/Editor/PluginBrowser/Source/PluginBrowser/Private/SPluginBrowser.cpp
@@ -7,16 +7,39 @@
 #include "SPluginCategoryTree.h"
 #include "SSearchBox.h"
 #include "PluginBrowserModule.h"
+#include "DirectoryWatcherModule.h"
 
 #define LOCTEXT_NAMESPACE "PluginsEditor"
 
-
-// @todo plugedit: Intro anim!!!
-
+SPluginBrowser::~SPluginBrowser()
+{
+	FDirectoryWatcherModule& DirectoryWatcherModule = FModuleManager::LoadModuleChecked<FDirectoryWatcherModule>(TEXT("DirectoryWatcher"));
+	for(auto& Pair: WatchDirectories)
+	{
+		DirectoryWatcherModule.Get()->UnregisterDirectoryChangedCallback_Handle(Pair.Key, Pair.Value);
+	}
+}
 
 void SPluginBrowser::Construct( const FArguments& Args )
 {
-	// @todo plugedit: Should we save/restore category selection and other view settings?  Splitter position, etc.
+	// Get the root directories which contain plugins
+	TArray<FString> WatchDirectoryNames;
+	WatchDirectoryNames.Add(FPaths::EnginePluginsDir());
+	if(FApp::HasGameName())
+	{
+		WatchDirectoryNames.Add(FPaths::GamePluginsDir());
+	}
+
+	// Add watchers for any change events on those directories
+	FDirectoryWatcherModule& DirectoryWatcherModule = FModuleManager::LoadModuleChecked<FDirectoryWatcherModule>(TEXT("DirectoryWatcher"));
+	for(const FString& WatchDirectoryName: WatchDirectoryNames)
+	{
+		FDelegateHandle Handle;
+		if(DirectoryWatcherModule.Get()->RegisterDirectoryChangedCallback_Handle(WatchDirectoryName, IDirectoryWatcher::FDirectoryChanged::CreateRaw(this, &SPluginBrowser::OnPluginDirectoryChanged), Handle))
+		{
+			WatchDirectories.Add(WatchDirectoryName, Handle);
+		}
+	}
 
 	RegisterActiveTimer (0.f, FWidgetActiveTimerDelegate::CreateSP (this, &SPluginBrowser::TriggerBreadcrumbRefresh));
 
@@ -179,6 +202,22 @@ void SPluginBrowser::SetNeedsRefresh()
 	RegisterActiveTimer (0.f, FWidgetActiveTimerDelegate::CreateSP (this, &SPluginBrowser::TriggerBreadcrumbRefresh));
 }
 
+void SPluginBrowser::OnPluginDirectoryChanged(const TArray<struct FFileChangeData>&)
+{
+	if(UpdatePluginsTimerHandle.IsValid())
+	{
+		UnRegisterActiveTimer(UpdatePluginsTimerHandle.ToSharedRef());
+	}
+	UpdatePluginsTimerHandle = RegisterActiveTimer(2.0f, FWidgetActiveTimerDelegate::CreateSP(this, &SPluginBrowser::UpdatePluginsTimerCallback));
+}
+
+EActiveTimerReturnType SPluginBrowser::UpdatePluginsTimerCallback(double InCurrentTime, float InDeltaTime)
+{
+	IPluginManager::Get().RefreshPluginsList();
+	SetNeedsRefresh();
+	return EActiveTimerReturnType::Stop;
+}
+
 EActiveTimerReturnType SPluginBrowser::TriggerBreadcrumbRefresh(double InCurrentTime, float InDeltaTime)
 {
 	RefreshBreadcrumbTrail();
diff --git a/Engine/Plugins/Editor/PluginBrowser/Source/PluginBrowser/Private/SPluginBrowser.h b/Engine/Plugins/Editor/PluginBrowser/Source/PluginBrowser/Private/SPluginBrowser.h
--- a/Engine/Plugins/Editor/PluginBrowser/Source/PluginBrowser/Private/SPluginBrowser.h
+++ b/Engine/Plugins/Editor/PluginBrowser/Source/PluginBrowser/Private/SPluginBrowser.h
@@ -16,6 +16,8 @@ class SPluginBrowser : public SCompoundWidget
 
 	SLATE_END_ARGS()
 
+	/** Destructor */
+	virtual ~SPluginBrowser();
 
 	/** Widget constructor */
 	void Construct( const FArguments& Args );
@@ -37,6 +39,12 @@ class SPluginBrowser : public SCompoundWidget
 
 private:
 
+	/** Called with notification that one of the plugin directories has changed */
+	void OnPluginDirectoryChanged(const TArray<struct FFileChangeData>& );
+
+	/** Timer callback for when */
+	EActiveTimerReturnType UpdatePluginsTimerCallback(double InCurrentTime, float InDeltaTime);
+
 	/** @return Is the "restart required" notice visible? */
 	EVisibility HandleRestartEditorNoticeVisibility() const;
 
@@ -56,6 +64,11 @@ private:
 	EActiveTimerReturnType TriggerBreadcrumbRefresh(double InCurrentTime, float InDeltaTime);
 
 private:
+	/** Handles to the directory changed delegates */
+	TMap<FString, FDelegateHandle> WatchDirectories;
+
+	/** Handles to the directory changed delegates */
+	TSharedPtr<FActiveTimerHandle> UpdatePluginsTimerHandle;
 
 	/** The plugin categories widget */
 	TSharedPtr< class SPluginCategoryTree > PluginCategories;
diff --git a/Engine/Plugins/Editor/PluginBrowser/Source/PluginBrowser/Private/SPluginCategoryTree.cpp b/Engine/Plugins/Editor/PluginBrowser/Source/PluginBrowser/Private/SPluginCategoryTree.cpp
--- a/Engine/Plugins/Editor/PluginBrowser/Source/PluginBrowser/Private/SPluginCategoryTree.cpp
+++ b/Engine/Plugins/Editor/PluginBrowser/Source/PluginBrowser/Private/SPluginCategoryTree.cpp
@@ -36,12 +36,6 @@ void SPluginCategoryTree::Construct( const FArguments& Args, const TSharedRef< S
 
 	RebuildAndFilterCategoryTree();
 
-	// Expand all the root categories by default
-	for(TSharedPtr<FPluginCategory> RootCategory: RootCategories)
-	{
-		TreeView->SetItemExpansion( RootCategory, true );
-	}
-
 	ChildSlot.AttachWidget( TreeView.ToSharedRef() );
 }
 
@@ -77,6 +71,7 @@ void SPluginCategoryTree::RebuildAndFilterCategoryTree()
 		for(TSharedPtr<FPluginCategory> Category: RootCategory->SubCategories)
 		{
 			Category->Plugins.Reset();
+			Category->SubCategories.Reset();
 		}
 	}
 
@@ -167,6 +162,12 @@ void SPluginCategoryTree::RebuildAndFilterCategoryTree()
 		RootCategory->SubCategories.Sort([](const TSharedPtr<FPluginCategory>& A, const TSharedPtr<FPluginCategory>& B) -> bool { return A->DisplayName.CompareTo(B->DisplayName) < 0; });
 	}
 
+	// Expand all the root categories by default
+	for(TSharedPtr<FPluginCategory> RootCategory: RootCategories)
+	{
+		TreeView->SetItemExpansion(RootCategory, true);
+	}
+
 	// Refresh the view
 	TreeView->RequestTreeRefresh();
 
diff --git a/Engine/Source/Runtime/Projects/Private/PluginManager.cpp b/Engine/Source/Runtime/Projects/Private/PluginManager.cpp
--- a/Engine/Source/Runtime/Projects/Private/PluginManager.cpp
+++ b/Engine/Source/Runtime/Projects/Private/PluginManager.cpp
@@ -107,200 +107,129 @@ FPluginManager::~FPluginManager()
 	//  explicitly unload your module somewhere.  We should be able to handle most cases automatically though!
 }
 
-void FPluginManager::DiscoverAllPlugins()
+void FPluginManager::RefreshPluginsList()
 {
-	struct Local
+	// Read a new list of all plugins
+	TArray<TSharedRef<FPlugin>> NewPlugins;
+	ReadAllPlugins(NewPlugins);
+
+	// Build a list of filenames for plugins which are enabled, and remove the rest
+	TArray<FString> EnabledPluginFileNames;
+	for(int32 Idx = 0; Idx < AllPlugins.Num(); Idx++)
 	{
+		const TSharedRef<FPlugin>& Plugin = AllPlugins[Idx];
+		if(Plugin->bEnabled)
+		{
+			EnabledPluginFileNames.Add(Plugin->FileName);
+		}
+		else
+		{
+			AllPlugins.RemoveAt(Idx--);
+		}
+	}
 
-	private:
-		/**
-		 * Recursively searches for plugins and generates a list of plugin descriptors
-		 *
-		 * @param	PluginsDirectory	Directory we're currently searching
-		 * @param	LoadedFrom			Where we're loading these plugins from (game, engine, etc)
-		 * @param	Plugins				The array to be filled in with new plugins including descriptors
-		 */
-		static void FindPluginsRecursively( const FString& PluginsDirectory, const EPluginLoadedFrom LoadedFrom, TArray< TSharedRef<FPlugin> >& Plugins )
+	// Add all the plugins which aren't already enabled
+	for(TSharedRef<FPlugin>& NewPlugin: NewPlugins)
+	{
+		if(!EnabledPluginFileNames.Contains(NewPlugin->FileName))
 		{
-			// NOTE: The logic in this function generally matches that of the C# code for FindPluginsRecursively
-			//       in UnrealBuildTool.  These routines should be kept in sync.
+			AllPlugins.Add(NewPlugin);
+		}
+	}
+}
 
-			// This directory scanning needs to be fast because it will happen every time at startup!  We don't
-			// want to blindly recurse down every subdirectory, because plugin content and code directories could
-			// contain a lot of files in total!
+void FPluginManager::DiscoverAllPlugins()
+{
+	ensure( AllPlugins.Num() == 0 );		// Should not have already been initialized!
+	ReadAllPlugins(AllPlugins);
 
-			// Each sub-directory is possibly a plugin.  If we find that it contains a plugin, we won't recurse any
-			// further -- you can't have plugins within plugins.  If we didn't find a plugin, we'll keep recursing.
-			TArray<FString> PossiblePluginDirectories;
-			{
-				class FPluginDirectoryVisitor : public IPlatformFile::FDirectoryVisitor
-				{
-				public:
-					FPluginDirectoryVisitor( TArray<FString>& InitFoundDirectories )
-						: FoundDirectories( InitFoundDirectories )
-					{
-					}
+	// Add the plugin binaries directory
+	for(const TSharedRef<FPlugin>& Plugin: AllPlugins)
+	{
+		const FString PluginBinariesPath = FPaths::Combine(*FPaths::GetPath(Plugin->FileName), TEXT("Binaries"), FPlatformProcess::GetBinariesSubdirectory());
+		FModuleManager::Get().AddBinariesDirectory(*PluginBinariesPath, Plugin->LoadedFrom == EPluginLoadedFrom::GameProject);
+	}
+}
 
-					virtual bool Visit(const TCHAR* FilenameOrDirectory, bool bIsDirectory) override
-					{
-						if( bIsDirectory )
-						{
-							FoundDirectories.Add( FString( FilenameOrDirectory ) );
-						}
+void FPluginManager::ReadAllPlugins(TArray<TSharedRef<FPlugin>>& Plugins)
+{
+#if (WITH_ENGINE && !IS_PROGRAM) || WITH_PLUGIN_SUPPORT
+	// Find "built-in" plugins.  That is, plugins situated right within the Engine directory.
+	ReadPluginsInDirectory(FPaths::EnginePluginsDir(), EPluginLoadedFrom::Engine, Plugins);
 
-						const bool bShouldContinue = true;
-						return bShouldContinue;
-					}
+	// Find plugins in the game project directory (<MyGameProject>/Plugins)
+	if( FApp::HasGameName() )
+	{
+		ReadPluginsInDirectory(FPaths::GamePluginsDir(), EPluginLoadedFrom::GameProject, Plugins);
+	}
+#endif
+}
 
-					TArray<FString>& FoundDirectories;
-				};
+void FPluginManager::ReadPluginsInDirectory(const FString& PluginsDirectory, const EPluginLoadedFrom LoadedFrom, TArray<TSharedRef<FPlugin>>& Plugins)
+{
+	// Make sure the directory even exists
+	if(FPlatformFileManager::Get().GetPlatformFile().DirectoryExists(*PluginsDirectory))
+	{
+		TArray<FString> FileNames;
+		FindPluginsInDirectory(PluginsDirectory, FileNames);
 
-				FPluginDirectoryVisitor PluginDirectoryVisitor( PossiblePluginDirectories );
-				FPlatformFileManager::Get().GetPlatformFile().IterateDirectory( *PluginsDirectory, PluginDirectoryVisitor );
+		for(const FString& FileName: FileNames)
+		{
+			FPluginDescriptor Descriptor;
+			FText FailureReason;
+			if(Descriptor.Load(FileName, FailureReason))
+			{
+				Plugins.Add(MakeShareable(new FPlugin(FileName, Descriptor, LoadedFrom)));
 			}
-				
-			for( auto PossiblePluginDirectoryIter( PossiblePluginDirectories.CreateConstIterator() ); PossiblePluginDirectoryIter; ++PossiblePluginDirectoryIter )
+			else
 			{
-				const FString PossiblePluginDirectory = *PossiblePluginDirectoryIter;
-
-				FString PluginDescriptorFilename;
-				{
-					// Usually the plugin descriptor is named the same as the directory.  It doesn't have to match the directory
-					// name but we'll check that first because its much faster than scanning!
-					const FString ProbablePluginDescriptorFilename = PossiblePluginDirectory / (FPaths::GetCleanFilename(PossiblePluginDirectory) + PluginSystemDefs::PluginDescriptorFileExtension);
-					if( FPlatformFileManager::Get().GetPlatformFile().FileExists( *ProbablePluginDescriptorFilename ) )
-					{
-						PluginDescriptorFilename = ProbablePluginDescriptorFilename;
-					}
-					else
-					{
-						// Scan the directory for a plugin descriptor.
-						TArray<FString> PluginDescriptorFilenames;
-						{
-							class FPluginDescriptorFileVisitor : public IPlatformFile::FDirectoryVisitor
-							{
-							public:
-								FPluginDescriptorFileVisitor( const FWildcardString& InitSearchPattern, TArray<FString>& InitFoundPluginDescriptorFilenames )
-									: SearchPattern( InitSearchPattern ),
-									  FoundPluginDescriptorFilenames( InitFoundPluginDescriptorFilenames )
-								{
-								}
-
-								virtual bool Visit(const TCHAR* FilenameOrDirectory, bool bIsDirectory) override
-								{
-									bool bShouldContinue = true;
-									if( !bIsDirectory )
-									{
-										if( SearchPattern.IsMatch( FilenameOrDirectory ) )
-										{
-											FoundPluginDescriptorFilenames.Add( FString( FilenameOrDirectory ) );
-											bShouldContinue = false;
-										}
-									}
-									return bShouldContinue;
-								}
-
-								const FWildcardString& SearchPattern;
-								TArray<FString>& FoundPluginDescriptorFilenames;
-							};
-
-							const FWildcardString PluginDescriptorSearchString = PossiblePluginDirectory / FString(TEXT("*")) + PluginSystemDefs::PluginDescriptorFileExtension;
-
-							FPluginDescriptorFileVisitor PluginDescriptorFileVisitor( PluginDescriptorSearchString, PluginDescriptorFilenames );
-							FPlatformFileManager::Get().GetPlatformFile().IterateDirectory( *PossiblePluginDirectory, PluginDescriptorFileVisitor );
-						}
-
-						// Did we find any plugin descriptor files?
-						if( PluginDescriptorFilenames.Num() > 0 )
-						{
-							PluginDescriptorFilename = PluginDescriptorFilenames[ 0 ];;
-						}
-					}
-				}
-
-				if( !PluginDescriptorFilename.IsEmpty() )
-				{
-					// Found a plugin directory!  No need to recurse any further.
-					FPluginDescriptor Descriptor;
-
-					// Load the descriptor
-					FText FailureReason;
-					if(Descriptor.Load(PluginDescriptorFilename, FailureReason))
-					{
-						TSharedRef< FPlugin > NewPlugin = MakeShareable( new FPlugin(PluginDescriptorFilename, Descriptor, LoadedFrom) );
-
-						if (FPaths::IsProjectFilePathSet())
-						{
-							FString GameProjectFolder = FPaths::GetPath(FPaths::GetProjectFilePath());
-							if (PluginDescriptorFilename.StartsWith(GameProjectFolder))
-							{
-								FString PluginBinariesFolder = FPaths::Combine(*FPaths::GetPath(PluginDescriptorFilename), TEXT("Binaries"), FPlatformProcess::GetBinariesSubdirectory());
-								FModuleManager::Get().AddBinariesDirectory(*PluginBinariesFolder, (LoadedFrom == EPluginLoadedFrom::GameProject));
-							}
-						}
-
-						Plugins.Add(NewPlugin);
-					}
-					else
-					{
-						// NOTE: Even though loading of this plugin failed, we'll keep processing other plugins
-						UE_LOG(LogPluginManager, Error, TEXT("%s (%s)"), *FailureReason.ToString(), *PluginDescriptorFilename);
-					}
-				}
-				else
-				{
-					// Didn't find a plugin in this directory.  Continue to look in subfolders.
-					FindPluginsRecursively( PossiblePluginDirectory, LoadedFrom, Plugins );
-				}
+				// NOTE: Even though loading of this plugin failed, we'll keep processing other plugins
+				UE_LOG(LogPluginManager, Error, TEXT("%s (%s)"), *FailureReason.ToString(), *FileName);
 			}
 		}
+	}
+}
 
+void FPluginManager::FindPluginsInDirectory(const FString& PluginsDirectory, TArray<FString>& FileNames)
+{
+	// Class to enumerate the contents of a directory, and find all sub-directories and plugin descriptors within it
+	struct FPluginDirectoryVisitor : public IPlatformFile::FDirectoryVisitor
+	{
+		TArray<FString> SubDirectories;
+		TArray<FString> PluginDescriptors;
 
-	public:
-
-		/**
-		 * Searches for plugins and generates a list of plugin descriptors
-		 *
-		 * @param	PluginsDirectory	The base directory to search for plugins
-		 * @param	LoadedFrom			Where we're loading these plugins from (game, engine, etc)
-		 * @param	Plugins				The array to be filled in with loaded plugin descriptors
-		 */
-		static void FindPluginsIn( const FString& PluginsDirectory, const EPluginLoadedFrom LoadedFrom, TArray< TSharedRef<FPlugin> >& Plugins )
+		virtual bool Visit(const TCHAR* FilenameOrDirectory, bool bIsDirectory) override
 		{
-			// Make sure the directory even exists
-			if( FPlatformFileManager::Get().GetPlatformFile().DirectoryExists( *PluginsDirectory ) )
+			FString FilenameOrDirectoryStr = FilenameOrDirectory;
+			if(bIsDirectory)
 			{
-				FindPluginsRecursively( PluginsDirectory, LoadedFrom, Plugins );
+				SubDirectories.Add(FilenameOrDirectoryStr);
 			}
+			else if(FilenameOrDirectoryStr.EndsWith(TEXT(".uplugin")))
+			{
+				PluginDescriptors.Add(FilenameOrDirectoryStr);
+			}
+			return true;
 		}
 	};
 
+	// Enumerate the contents of the current directory
+	FPluginDirectoryVisitor Visitor;
+	FPlatformFileManager::Get().GetPlatformFile().IterateDirectory(*PluginsDirectory, Visitor);
 
+	// If there's no plugins in this directory, recurse through all the child directories
+	if(Visitor.PluginDescriptors.Num() == 0)
 	{
-		TArray< TSharedRef<FPlugin> > Plugins;
-
-#if (WITH_ENGINE && !IS_PROGRAM) || WITH_PLUGIN_SUPPORT
-		// Find "built-in" plugins.  That is, plugins situated right within the Engine directory.
-		Local::FindPluginsIn( FPaths::EnginePluginsDir(), EPluginLoadedFrom::Engine, Plugins );
-
-		// Find plugins in the game project directory (<MyGameProject>/Plugins)
-		if( FApp::HasGameName() )
+		for(const FString& SubDirectory: Visitor.SubDirectories)
 		{
-			Local::FindPluginsIn( FPaths::GamePluginsDir(), EPluginLoadedFrom::GameProject, Plugins );
+			FindPluginsInDirectory(SubDirectory, FileNames);
 		}
-#endif		// (WITH_ENGINE && !IS_PROGRAM) || WITH_PLUGIN_SUPPORT
-
-
-		// Create plugin objects for all of the plugins that we found
-		ensure( AllPlugins.Num() == 0 );		// Should not have already been initialized!
-		AllPlugins = Plugins;
-		for( auto PluginIt( Plugins.CreateConstIterator() ); PluginIt; ++PluginIt )
+	}
+	else
+	{
+		for(const FString& PluginDescriptor: Visitor.PluginDescriptors)
 		{
-			const TSharedRef<FPlugin>& Plugin = *PluginIt;
-
-			// Add the plugin binaries directory
-			const FString PluginBinariesPath = FPaths::Combine(*FPaths::GetPath(Plugin->FileName), TEXT("Binaries"), FPlatformProcess::GetBinariesSubdirectory());
-			FModuleManager::Get().AddBinariesDirectory(*PluginBinariesPath, Plugin->LoadedFrom == EPluginLoadedFrom::GameProject);
+			FileNames.Add(PluginDescriptor);
 		}
 	}
 }
diff --git a/Engine/Source/Runtime/Projects/Private/PluginManager.h b/Engine/Source/Runtime/Projects/Private/PluginManager.h
--- a/Engine/Source/Runtime/Projects/Private/PluginManager.h
+++ b/Engine/Source/Runtime/Projects/Private/PluginManager.h
@@ -61,6 +61,7 @@ public:
 	~FPluginManager();
 
 	/** IPluginManager interface */
+	virtual void RefreshPluginsList();
 	virtual bool LoadModulesForEnabledPlugins( const ELoadingPhase::Type LoadingPhase ) override;
 	virtual void SetRegisterMountPointDelegate( const FRegisterMountPointDelegate& Delegate ) override;
 	virtual bool AreRequiredPluginsAvailable() override;
@@ -76,6 +77,15 @@ private:
 	    This is called when the plugin manager singleton is first accessed. */
 	void DiscoverAllPlugins();
 
+	/** Reads all the plugin descriptors */
+	static void ReadAllPlugins(TArray<TSharedRef<FPlugin>>& Plugins);
+
+	/** Reads all the plugin descriptors from disk */
+	static void ReadPluginsInDirectory(const FString& PluginsDirectory, const EPluginLoadedFrom LoadedFrom, TArray<TSharedRef<FPlugin>>& Plugins);
+
+	/** Finds all the plugin descriptors underneath a given directory */
+	static void FindPluginsInDirectory(const FString& PluginsDirectory, TArray<FString>& FileNames);
+
 	/** Sets the bPluginEnabled flag on all plugins found from DiscoverAllPlugins that are enabled in config */
 	bool ConfigureEnabledPlugins();
 
diff --git a/Engine/Source/Runtime/Projects/Public/Interfaces/IPluginManager.h b/Engine/Source/Runtime/Projects/Public/Interfaces/IPluginManager.h
--- a/Engine/Source/Runtime/Projects/Public/Interfaces/IPluginManager.h
+++ b/Engine/Source/Runtime/Projects/Public/Interfaces/IPluginManager.h
@@ -129,6 +129,11 @@ class IPluginManager
 {
 public:
 
+	/** 
+	 * Updates the list of plugins.
+	 */
+	virtual void RefreshPluginsList() = 0;
+
 	/**
 	 * Loads all plug-ins
 	 *
