Always enable installed plugins by default, but update your project file with references to them on startup (so other users opening the project will be able to find it on the marketplace).

diff --git a/Engine/Source/Editor/GameProjectGeneration/Private/GameProjectUtils.cpp b/Engine/Source/Editor/GameProjectGeneration/Private/GameProjectUtils.cpp
--- a/Engine/Source/Editor/GameProjectGeneration/Private/GameProjectUtils.cpp
+++ b/Engine/Source/Editor/GameProjectGeneration/Private/GameProjectUtils.cpp
@@ -684,6 +684,30 @@ void GameProjectUtils::CheckForOutOfDateGameProjectFile()
 				}
 			}
 		}
+
+		// Check if there are any installed plugins which aren't referenced by the project file
+		if(!UpdateGameProjectNotification.IsValid())
+		{
+			const FProjectDescriptor* Project = IProjectManager::Get().GetCurrentProject();
+			if(Project != nullptr)
+			{
+				TArray<FPluginReferenceDescriptor> NewPluginReferences;
+				for(TSharedRef<IPlugin>& Plugin: IPluginManager::Get().GetEnabledPlugins())
+				{
+					if(Plugin->GetDescriptor().bInstalled && Project->FindPluginReferenceIndex(Plugin->GetName()) == INDEX_NONE)
+					{
+						FPluginReferenceDescriptor PluginReference(Plugin->GetName(), true, Plugin->GetDescriptor().MarketplaceURL);
+						NewPluginReferences.Add(PluginReference);
+					}
+				}
+				if(NewPluginReferences.Num() > 0)
+				{
+					UpdateProject(FProjectDescriptorModifier::CreateLambda( 
+						[NewPluginReferences](FProjectDescriptor& Descriptor){ Descriptor.Plugins.Append(NewPluginReferences); return true; }
+					));
+				}
+			}
+		}
 	}
 }
 
diff --git a/Engine/Source/Programs/AutomationTool/Scripts/BuildPluginCommand.Automation.cs b/Engine/Source/Programs/AutomationTool/Scripts/BuildPluginCommand.Automation.cs
--- a/Engine/Source/Programs/AutomationTool/Scripts/BuildPluginCommand.Automation.cs
+++ b/Engine/Source/Programs/AutomationTool/Scripts/BuildPluginCommand.Automation.cs
@@ -139,7 +139,7 @@ class BuildPlugin : BuildCommand
 		// Get the output plugin filename
 		string TargetPluginFileName = CommandUtils.MakeRerootedFilePath(Path.GetFullPath(PluginFileName), Path.GetDirectoryName(Path.GetFullPath(PluginFileName)), PackageDirectory);
 		PluginDescriptor NewDescriptor = PluginDescriptor.FromFile(TargetPluginFileName);
-		NewDescriptor.bEnabledByDefault = false;
+		NewDescriptor.bEnabledByDefault = true;
 		NewDescriptor.bInstalled = true;
 		NewDescriptor.Save(TargetPluginFileName);
 	}
diff --git a/Engine/Source/Runtime/Projects/Private/PluginManager.cpp b/Engine/Source/Runtime/Projects/Private/PluginManager.cpp
--- a/Engine/Source/Runtime/Projects/Private/PluginManager.cpp
+++ b/Engine/Source/Runtime/Projects/Private/PluginManager.cpp
@@ -279,6 +279,7 @@ bool FPluginManager::ConfigureEnabledPlugins()
 							FString Error;
 							FPlatformProcess::LaunchURL(*Plugin.MarketplaceURL, nullptr, &Error);
 							if(Error.Len() > 0) FMessageDialog::Open(EAppMsgType::Ok, FText::FromString(Error));
+							return false;
 						}
 					}
 					else
@@ -286,7 +287,17 @@ bool FPluginManager::ConfigureEnabledPlugins()
 						FString Description = (Plugin.Description.Len() > 0)? FString::Printf(TEXT("\n\n%s"), *Plugin.Description) : FString();
 						FMessageDialog::Open(EAppMsgType::Ok, FText::Format(LOCTEXT("PluginMissingError", "This project requires the {0} plugin. {1}"), FText::FromString(Plugin.Name), FText::FromString(Description)), &Caption);
 					}
-					return false;
+
+					if(FMessageDialog::Open(EAppMsgType::YesNo, FText::Format(LOCTEXT("PluginMissingDisable", "Would you like to disable {0}? You will no longer be able to open any assets created using it."), FText::FromString(Plugin.Name)), &Caption) == EAppReturnType::No)
+					{
+						return false;
+					}
+
+					FText FailReason;
+					if(!IProjectManager::Get().SetPluginEnabled(*Plugin.Name, false, FailReason))
+					{
+						FMessageDialog::Open(EAppMsgType::Ok, FailReason);
+					}
 				}
 			}
 		}
diff --git a/Engine/Source/Runtime/Projects/Private/ProjectDescriptor.cpp b/Engine/Source/Runtime/Projects/Private/ProjectDescriptor.cpp
--- a/Engine/Source/Runtime/Projects/Private/ProjectDescriptor.cpp
+++ b/Engine/Source/Runtime/Projects/Private/ProjectDescriptor.cpp
@@ -21,6 +21,18 @@ bool FProjectDescriptor::IsSigned(const FString& FilePath) const
 	return EpicSampleNameHash == GetTypeHash(FPaths::GetCleanFilename(FilePath));
 }
 
+int32 FProjectDescriptor::FindPluginReferenceIndex(const FString& PluginName) const
+{
+	for(int32 Idx = 0; Idx < Plugins.Num(); Idx++)
+	{
+		if(Plugins[Idx].Name == PluginName)
+		{
+			return Idx;
+		}
+	}
+	return INDEX_NONE;
+}
+
 void FProjectDescriptor::UpdateSupportedTargetPlatforms(const FName& InPlatformName, bool bIsSupported)
 {
 	if ( bIsSupported )
diff --git a/Engine/Source/Runtime/Projects/Private/ProjectManager.cpp b/Engine/Source/Runtime/Projects/Private/ProjectManager.cpp
--- a/Engine/Source/Runtime/Projects/Private/ProjectManager.cpp
+++ b/Engine/Source/Runtime/Projects/Private/ProjectManager.cpp
@@ -239,7 +239,7 @@ void FProjectManager::ClearSupportedTargetPlatformsForCurrentProject()
 void FProjectManager::GetEnabledPlugins(TArray<FString>& OutPluginNames) const
 {
 	// Get the default list of plugin names
-	GetDefaultEnabledPlugins(OutPluginNames);
+	GetDefaultEnabledPlugins(OutPluginNames, true);
 
 	// Modify that with the list of plugins in the project file
 	const FProjectDescriptor *Project = GetCurrentProject();
@@ -306,7 +306,7 @@ bool FProjectManager::SetPluginEnabled(const FString& PluginName, bool bEnabled,
 	{
 		// Get the default list of enabled plugins
 		TArray<FString> DefaultEnabledPlugins;
-		GetDefaultEnabledPlugins(DefaultEnabledPlugins);
+		GetDefaultEnabledPlugins(DefaultEnabledPlugins, false);
 
 		// Check the enabled state is the same in that
 		if(DefaultEnabledPlugins.Contains(PluginName) == bEnabled)
@@ -325,7 +325,7 @@ bool FProjectManager::SetPluginEnabled(const FString& PluginName, bool bEnabled,
 	return true;
 }
 
-void FProjectManager::GetDefaultEnabledPlugins(TArray<FString>& OutPluginNames)
+void FProjectManager::GetDefaultEnabledPlugins(TArray<FString>& OutPluginNames, bool bIncludeInstalledPlugins)
 {
 	// Add all the game plugins and everything marked as enabled by default
 	TArray<FPluginStatus> PluginStatuses = IPluginManager::Get().QueryStatusForAllPlugins();
@@ -333,7 +333,10 @@ void FProjectManager::GetDefaultEnabledPlugins(TArray<FString>& OutPluginNames)
 	{
 		if(PluginStatus.Descriptor.bEnabledByDefault || PluginStatus.LoadedFrom == EPluginLoadedFrom::GameProject)
 		{
-			OutPluginNames.AddUnique(PluginStatus.Name);
+			if(bIncludeInstalledPlugins || !PluginStatus.Descriptor.bInstalled)
+			{
+				OutPluginNames.AddUnique(PluginStatus.Name);
+			}
 		}
 	}
 }
diff --git a/Engine/Source/Runtime/Projects/Private/ProjectManager.h b/Engine/Source/Runtime/Projects/Private/ProjectManager.h
--- a/Engine/Source/Runtime/Projects/Private/ProjectManager.h
+++ b/Engine/Source/Runtime/Projects/Private/ProjectManager.h
@@ -34,7 +34,7 @@ private:
 	static void QueryStatusForProjectImpl(const FProjectDescriptor& Project, const FString& FilePath, FProjectStatus& OutProjectStatus);
 
 	/** Gets the list of plugins enabled by default, excluding the project overrides */
-	static void GetDefaultEnabledPlugins(TArray<FString>& PluginNames);
+	static void GetDefaultEnabledPlugins(TArray<FString>& PluginNames, bool bIncludeInstalledPlugins);
 
 	/** The project that is currently loaded in the editor */
 	TSharedPtr< FProjectDescriptor > CurrentProject;
diff --git a/Engine/Source/Runtime/Projects/Public/ProjectDescriptor.h b/Engine/Source/Runtime/Projects/Public/ProjectDescriptor.h
--- a/Engine/Source/Runtime/Projects/Public/ProjectDescriptor.h
+++ b/Engine/Source/Runtime/Projects/Public/ProjectDescriptor.h
@@ -90,6 +90,9 @@ struct PROJECTS_API FProjectDescriptor
 	/** Checks whether the descriptor is signed */
 	bool IsSigned(const FString& FilePath) const;
 
+	/** Finds the index of a plugin in the references array */
+	int32 FindPluginReferenceIndex(const FString& PluginName) const;
+
 	/** Updates the supported target platforms list */
 	void UpdateSupportedTargetPlatforms(const FName& InPlatformName, bool bIsSupported);
 
