Fix disabling an installed plugin removing it from the plugin reference list, which causes them to be re-enabled on next engine startup.

diff --git a/Engine/Source/Runtime/Projects/Private/ProjectManager.cpp b/Engine/Source/Runtime/Projects/Private/ProjectManager.cpp
--- a/Engine/Source/Runtime/Projects/Private/ProjectManager.cpp
+++ b/Engine/Source/Runtime/Projects/Private/ProjectManager.cpp
@@ -304,15 +304,20 @@ bool FProjectManager::SetPluginEnabled(const FString& PluginName, bool bEnabled,
 	const FPluginReferenceDescriptor* PluginRef = &CurrentProject->Plugins[PluginRefIdx];
 	if(PluginRef->WhitelistPlatforms.Num() == 0 && PluginRef->BlacklistPlatforms.Num() == 0)
 	{
-		// Get the default list of enabled plugins
-		TArray<FString> DefaultEnabledPlugins;
-		GetDefaultEnabledPlugins(DefaultEnabledPlugins, false);
-
-		// Check the enabled state is the same in that
-		if(DefaultEnabledPlugins.Contains(PluginName) == bEnabled)
+		// We alway need to be explicit about installed plugins, because they'll be auto-enabled again if we're not.
+		TSharedPtr<IPlugin> Plugin = IPluginManager::Get().FindPlugin(PluginName);
+		if(!Plugin.IsValid() || !Plugin->GetDescriptor().bInstalled)
 		{
-			CurrentProject->Plugins.RemoveAt(PluginRefIdx);
-			PluginRefIdx = INDEX_NONE;
+			// Get the default list of enabled plugins
+			TArray<FString> DefaultEnabledPlugins;
+			GetDefaultEnabledPlugins(DefaultEnabledPlugins, false);
+
+			// Check the enabled state is the same in that
+			if(DefaultEnabledPlugins.Contains(PluginName) == bEnabled)
+			{
+				CurrentProject->Plugins.RemoveAt(PluginRefIdx);
+				PluginRefIdx = INDEX_NONE;
+			}
 		}
 	}
 
