Ensure that dependencies on removed plugins are removed from the saved .uproject file

Removed Saving from the Project Manager's SetPluginEnabled function, setting a dirty flag instead, as it won't be able to check source control
Added a new function to save the current project from the Project Manager and clear the dirty flag on success

diff --git a/Engine/Plugins/Editor/PluginBrowser/Source/PluginBrowser/Private/SPluginTile.cpp b/Engine/Plugins/Editor/PluginBrowser/Source/PluginBrowser/Private/SPluginTile.cpp
--- a/Engine/Plugins/Editor/PluginBrowser/Source/PluginBrowser/Private/SPluginTile.cpp
+++ b/Engine/Plugins/Editor/PluginBrowser/Source/PluginBrowser/Private/SPluginTile.cpp
@@ -387,7 +387,6 @@ void SPluginTile::OnEnablePluginCheckboxChanged(ECheckBoxState NewCheckedState)
 		}
 	}
 
-	FGameProjectGenerationModule::Get().TryMakeProjectFileWriteable(FPaths::GetProjectFilePath());
 	FText FailMessage;
 
 	if (!IProjectManager::Get().SetPluginEnabled(Plugin->GetName(), bNewEnabledState, FailMessage, PluginDescriptor.MarketplaceURL))
@@ -396,14 +395,23 @@ void SPluginTile::OnEnablePluginCheckboxChanged(ECheckBoxState NewCheckedState)
 	}
 	else
 	{
-		FPluginBrowserModule& PluginBrowserModule = FPluginBrowserModule::Get();
-		if(Plugin->IsEnabled() == bNewEnabledState)
+		FGameProjectGenerationModule::Get().TryMakeProjectFileWriteable(FPaths::GetProjectFilePath());
+
+		if (!IProjectManager::Get().SaveCurrentProjectToDisk(FailMessage))
 		{
-			PluginBrowserModule.PendingEnablePlugins.Remove(Plugin->GetName());
+			FMessageDialog::Open(EAppMsgType::Ok, FailMessage);
 		}
 		else
 		{
-			PluginBrowserModule.PendingEnablePlugins.FindOrAdd(Plugin->GetName()) = bNewEnabledState;
+			FPluginBrowserModule& PluginBrowserModule = FPluginBrowserModule::Get();
+			if(Plugin->IsEnabled() == bNewEnabledState)
+			{
+				PluginBrowserModule.PendingEnablePlugins.Remove(Plugin->GetName());
+			}
+			else
+			{
+				PluginBrowserModule.PendingEnablePlugins.FindOrAdd(Plugin->GetName()) = bNewEnabledState;
+			}
 		}
 	}
 
diff --git a/Engine/Source/Editor/GameProjectGeneration/Private/GameProjectUtils.cpp b/Engine/Source/Editor/GameProjectGeneration/Private/GameProjectUtils.cpp
--- a/Engine/Source/Editor/GameProjectGeneration/Private/GameProjectUtils.cpp
+++ b/Engine/Source/Editor/GameProjectGeneration/Private/GameProjectUtils.cpp
@@ -652,6 +652,16 @@ void GameProjectUtils::CheckForOutOfDateGameProjectFile()
 {
 	if ( FPaths::IsProjectFilePathSet() )
 	{
+		if (IProjectManager::Get().IsCurrentProjectDirty())
+		{
+			FText FailMessage;
+			TryMakeProjectFileWriteable(FPaths::GetProjectFilePath());
+			if (!IProjectManager::Get().SaveCurrentProjectToDisk(FailMessage))
+			{
+				FMessageDialog::Open(EAppMsgType::Ok, FailMessage);
+			}
+		}
+
 		FProjectStatus ProjectStatus;
 		if (IProjectManager::Get().QueryStatusForCurrentProject(ProjectStatus))
 		{
diff --git a/Engine/Source/Runtime/Engine/Private/Commandlets/PluginCommandlet.cpp b/Engine/Source/Runtime/Engine/Private/Commandlets/PluginCommandlet.cpp
--- a/Engine/Source/Runtime/Engine/Private/Commandlets/PluginCommandlet.cpp
+++ b/Engine/Source/Runtime/Engine/Private/Commandlets/PluginCommandlet.cpp
@@ -81,6 +81,12 @@ int32 UPluginCommandlet::Main(const FString& Params)
 					if (IPluginManager::Get().FindPlugin(PluginName).IsValid() || CmdLineSwitches.Contains(TEXT("FORCE")))
 					{
 						bSetPluginState = IProjectManager::Get().SetPluginEnabled(PluginName, bEnable, FailReason);
+
+						// Try to save the project file
+						if (!IProjectManager::Get().SaveCurrentProjectToDisk(FailReason))
+						{
+							bSetPluginState = false;
+						}
 					}
 					else
 					{
diff --git a/Engine/Source/Runtime/Projects/Private/ProjectManager.cpp b/Engine/Source/Runtime/Projects/Private/ProjectManager.cpp
--- a/Engine/Source/Runtime/Projects/Private/ProjectManager.cpp
+++ b/Engine/Source/Runtime/Projects/Private/ProjectManager.cpp
@@ -7,6 +7,7 @@ DEFINE_LOG_CATEGORY_STATIC( LogProjectManager, Log, All );
 #define LOCTEXT_NAMESPACE "ProjectManager"
 
 FProjectManager::FProjectManager()
+	: bIsCurrentProjectDirty(false)
 {
 }
 
@@ -321,11 +322,8 @@ bool FProjectManager::SetPluginEnabled(const FString& PluginName, bool bEnabled,
 		}
 	}
 
-	// Try to save the project file
-	if(!CurrentProject->Save(*FPaths::GetProjectFilePath(), OutFailReason))
-	{
-		return false;
-	}
+	// Mark project as dirty
+	bIsCurrentProjectDirty = true;
 
 	return true;
 }
@@ -346,6 +344,22 @@ void FProjectManager::GetDefaultEnabledPlugins(TArray<FString>& OutPluginNames,
 	}
 }
 
+bool FProjectManager::IsCurrentProjectDirty() const
+{
+	return bIsCurrentProjectDirty;
+}
+
+bool FProjectManager::SaveCurrentProjectToDisk(FText& OutFailReason)
+{
+	if (CurrentProject->Save(FPaths::GetProjectFilePath(), OutFailReason))
+	{
+		bIsCurrentProjectDirty = false;
+		return true;
+	}
+	return false;
+}
+
+
 IProjectManager& IProjectManager::Get()
 {
 	// Single instance of manager, allocated on demand and destroyed on program exit.
diff --git a/Engine/Source/Runtime/Projects/Private/ProjectManager.h b/Engine/Source/Runtime/Projects/Private/ProjectManager.h
--- a/Engine/Source/Runtime/Projects/Private/ProjectManager.h
+++ b/Engine/Source/Runtime/Projects/Private/ProjectManager.h
@@ -29,6 +29,8 @@ public:
 	virtual void GetEnabledPlugins(TArray<FString>& OutPluginNames) const override;
 	virtual bool IsNonDefaultPluginEnabled() const override;
 	virtual bool SetPluginEnabled(const FString& PluginName, bool bEnabled, FText& OutFailReason, const FString& MarketplaceURL) override;
+	virtual bool IsCurrentProjectDirty() const override;
+	virtual bool SaveCurrentProjectToDisk(FText& OutFailReason) override;
 
 private:
 	static void QueryStatusForProjectImpl(const FProjectDescriptor& Project, const FString& FilePath, FProjectStatus& OutProjectStatus);
@@ -39,6 +41,9 @@ private:
 	/** The project that is currently loaded in the editor */
 	TSharedPtr< FProjectDescriptor > CurrentProject;
 
+	/** Whether the current project has been modified but not saved to disk */
+	bool bIsCurrentProjectDirty;
+
 	/** Delegate called when the target platforms for the current project are changed */
 	FOnTargetPlatformsForCurrentProjectChangedEvent OnTargetPlatformsForCurrentProjectChangedEvent;
 };
diff --git a/Engine/Source/Runtime/Projects/Public/Interfaces/IProjectManager.h b/Engine/Source/Runtime/Projects/Public/Interfaces/IProjectManager.h
--- a/Engine/Source/Runtime/Projects/Public/Interfaces/IProjectManager.h
+++ b/Engine/Source/Runtime/Projects/Public/Interfaces/IProjectManager.h
@@ -193,13 +193,26 @@ public:
 	virtual bool IsNonDefaultPluginEnabled() const = 0;
 
 	/**
-	 * Sets whether a plugin is enabled, and updates the project descriptor on disk. May require restarting to load it.
+	 * Sets whether a plugin is enabled, and updates the current project descriptor. Does not save to disk and may require restarting to load it.
 	 * 
 	 * @param	PluginName		Name of the plugin
 	 * @param	bEnabled		Whether to enable or disable the plugin
 	 * @param	OutFailReason	On failure, gives an error message
 	 * @param	MarketplaceURL	Marketplace URL to open if the user does not have this plugin installed
-	 * @return	True if the plugin has been marked as enabled, and the project descriptor has been saved.
+	 * @return	True if the plugin has been marked as enabled, and the project descriptor has been updated.
 	 */
 	virtual bool SetPluginEnabled(const FString& PluginName, bool bEnabled, FText& OutFailReason, const FString& MarketplaceURL = TEXT("")) = 0;
+
+	/**
+	 * Checks whether the current loaded project has been modified but not saved to disk
+	 */
+	virtual bool IsCurrentProjectDirty() const = 0;
+
+	/**
+	 * Saves the current project to the project path
+	 *
+	 * @param	OutFailReason	On failure, gives an error message
+	 * @return	True if the project was saved successfully
+	 */
+	virtual bool SaveCurrentProjectToDisk(FText& OutFailReason) = 0;
 };
\ No newline at end of file
