Support for silently ignoring plugin load failures. Can happen if plugin not supported on all platforms.

Merging 2587926

//depot/UE4-UT/...

to //depot/UE4/...

diff --git a/Engine/Source/Runtime/Projects/Private/PluginDescriptor.cpp b/Engine/Source/Runtime/Projects/Private/PluginDescriptor.cpp
--- a/Engine/Source/Runtime/Projects/Private/PluginDescriptor.cpp
+++ b/Engine/Source/Runtime/Projects/Private/PluginDescriptor.cpp
@@ -158,6 +158,7 @@ FString FPluginDescriptor::ToString() const
 FPluginReferenceDescriptor::FPluginReferenceDescriptor( const FString& InName, bool bInEnabled, const FString& InMarketplaceURL )
 	: Name(InName)
 	, bEnabled(bInEnabled)
+	, bNotRequired(false)
 	, MarketplaceURL(InMarketplaceURL)
 { }
 
@@ -201,6 +202,8 @@ bool FPluginReferenceDescriptor::Read( const FJsonObject& Object, FText& OutFail
 		OutFailReason = LOCTEXT("PluginReferenceWithoutEnabled", "Plugin references must have an 'Enabled' field");
 		return false;
 	}
+	
+	Object.TryGetBoolField(TEXT("NotRequired"), bNotRequired);
 
 	// Read the metadata for users that don't have the plugin installed
 	Object.TryGetStringField(TEXT("Description"), Description);
@@ -247,6 +250,7 @@ void FPluginReferenceDescriptor::Write( TJsonWriter<>& Writer ) const
 	Writer.WriteObjectStart();
 	Writer.WriteValue(TEXT("Name"), Name);
 	Writer.WriteValue(TEXT("Enabled"), bEnabled);
+	Writer.WriteValue(TEXT("NotRequired"), bNotRequired);
 
 	if (Description.Len() > 0)
 	{
diff --git a/Engine/Source/Runtime/Projects/Private/PluginManager.cpp b/Engine/Source/Runtime/Projects/Private/PluginManager.cpp
--- a/Engine/Source/Runtime/Projects/Private/PluginManager.cpp
+++ b/Engine/Source/Runtime/Projects/Private/PluginManager.cpp
@@ -282,21 +282,26 @@ bool FPluginManager::ConfigureEnabledPlugins()
 							return false;
 						}
 					}
-					else
+					else if (Plugin.bNotRequired)
 					{
-						FString Description = (Plugin.Description.Len() > 0)? FString::Printf(TEXT("\n\n%s"), *Plugin.Description) : FString();
-						FMessageDialog::Open(EAppMsgType::Ok, FText::Format(LOCTEXT("PluginMissingError", "This project requires the {0} plugin. {1}"), FText::FromString(Plugin.Name), FText::FromString(Description)), &Caption);
+						FText PluginFailReason = FText::Format(LOCTEXT("PluginMissingNotRequiredError", "This project is missing the {0} plugin, disabling it as it is not required."), FText::FromString(Plugin.Name));
+						IProjectManager::Get().SetPluginEnabled(*Plugin.Name, false, PluginFailReason);
 					}
-
-					if(FMessageDialog::Open(EAppMsgType::YesNo, FText::Format(LOCTEXT("PluginMissingDisable", "Would you like to disable {0}? You will no longer be able to open any assets created using it."), FText::FromString(Plugin.Name)), &Caption) == EAppReturnType::No)
+					else
 					{
-						return false;
-					}
+						FString Description = (Plugin.Description.Len() > 0) ? FString::Printf(TEXT("\n\n%s"), *Plugin.Description) : FString();
+						FMessageDialog::Open(EAppMsgType::Ok, FText::Format(LOCTEXT("PluginMissingError", "This project requires the {0} plugin. {1}"), FText::FromString(Plugin.Name), FText::FromString(Description)), &Caption);
+						
+						if (FMessageDialog::Open(EAppMsgType::YesNo, FText::Format(LOCTEXT("PluginMissingDisable", "Would you like to disable {0}? You will no longer be able to open any assets created using it."), FText::FromString(Plugin.Name)), &Caption) == EAppReturnType::No)
+						{
+							return false;
+						}
 
-					FText FailReason;
-					if(!IProjectManager::Get().SetPluginEnabled(*Plugin.Name, false, FailReason))
-					{
-						FMessageDialog::Open(EAppMsgType::Ok, FailReason);
+						FText FailReason;
+						if (!IProjectManager::Get().SetPluginEnabled(*Plugin.Name, false, FailReason))
+						{
+							FMessageDialog::Open(EAppMsgType::Ok, FailReason);
+						}
 					}
 				}
 			}
diff --git a/Engine/Source/Runtime/Projects/Public/PluginDescriptor.h b/Engine/Source/Runtime/Projects/Public/PluginDescriptor.h
--- a/Engine/Source/Runtime/Projects/Public/PluginDescriptor.h
+++ b/Engine/Source/Runtime/Projects/Public/PluginDescriptor.h
@@ -106,6 +106,9 @@ struct PROJECTS_API FPluginReferenceDescriptor
 	/** Whether it should be enabled by default */
 	bool bEnabled;
 
+	/** Whether we should silently ignore missing plugins */
+	bool bNotRequired;
+
 	/** Description of the plugin for users that do not have it installed. */
 	FString Description;
 
