Reverting plugin descriptor change

Merging 2590240

//depot/UE4-UT/...

to //depot/UE4/...

diff --git a/Engine/Source/Runtime/Projects/Private/PluginDescriptor.cpp b/Engine/Source/Runtime/Projects/Private/PluginDescriptor.cpp
--- a/Engine/Source/Runtime/Projects/Private/PluginDescriptor.cpp
+++ b/Engine/Source/Runtime/Projects/Private/PluginDescriptor.cpp
@@ -158,7 +158,6 @@ FString FPluginDescriptor::ToString() const
 FPluginReferenceDescriptor::FPluginReferenceDescriptor( const FString& InName, bool bInEnabled, const FString& InMarketplaceURL )
 	: Name(InName)
 	, bEnabled(bInEnabled)
-	, bNotRequired(false)
 	, MarketplaceURL(InMarketplaceURL)
 { }
 
@@ -203,8 +202,6 @@ bool FPluginReferenceDescriptor::Read( const FJsonObject& Object, FText& OutFail
 		return false;
 	}
 	
-	Object.TryGetBoolField(TEXT("NotRequired"), bNotRequired);
-
 	// Read the metadata for users that don't have the plugin installed
 	Object.TryGetStringField(TEXT("Description"), Description);
 	Object.TryGetStringField(TEXT("MarketplaceURL"), MarketplaceURL);
@@ -250,7 +247,6 @@ void FPluginReferenceDescriptor::Write( TJsonWriter<>& Writer ) const
 	Writer.WriteObjectStart();
 	Writer.WriteValue(TEXT("Name"), Name);
 	Writer.WriteValue(TEXT("Enabled"), bEnabled);
-	Writer.WriteValue(TEXT("NotRequired"), bNotRequired);
 
 	if (Description.Len() > 0)
 	{
diff --git a/Engine/Source/Runtime/Projects/Private/PluginManager.cpp b/Engine/Source/Runtime/Projects/Private/PluginManager.cpp
--- a/Engine/Source/Runtime/Projects/Private/PluginManager.cpp
+++ b/Engine/Source/Runtime/Projects/Private/PluginManager.cpp
@@ -282,11 +282,6 @@ bool FPluginManager::ConfigureEnabledPlugins()
 							return false;
 						}
 					}
-					else if (Plugin.bNotRequired)
-					{
-						FText PluginFailReason = FText::Format(LOCTEXT("PluginMissingNotRequiredError", "This project is missing the {0} plugin, disabling it as it is not required."), FText::FromString(Plugin.Name));
-						IProjectManager::Get().SetPluginEnabled(*Plugin.Name, false, PluginFailReason);
-					}
 					else
 					{
 						FString Description = (Plugin.Description.Len() > 0) ? FString::Printf(TEXT("\n\n%s"), *Plugin.Description) : FString();
diff --git a/Engine/Source/Runtime/Projects/Public/PluginDescriptor.h b/Engine/Source/Runtime/Projects/Public/PluginDescriptor.h
--- a/Engine/Source/Runtime/Projects/Public/PluginDescriptor.h
+++ b/Engine/Source/Runtime/Projects/Public/PluginDescriptor.h
@@ -105,10 +105,7 @@ struct PROJECTS_API FPluginReferenceDescriptor
 
 	/** Whether it should be enabled by default */
 	bool bEnabled;
-
-	/** Whether we should silently ignore missing plugins */
-	bool bNotRequired;
-
+	
 	/** Description of the plugin for users that do not have it installed. */
 	FString Description;
 
