Added a check on whether plugins are compatible when dealing with a content only project as they will not try to build them later. Give the option to disable them before they are added to the list of plugins to load.

Also fixed issue with iteration of plugins as items could be removed from the list while doing so.

UE-17136 - Enabling the WinDualShock plugin will cause the project to be unable to open again

diff --git a/Engine/Source/Runtime/Projects/Private/PluginManager.cpp b/Engine/Source/Runtime/Projects/Private/PluginManager.cpp
--- a/Engine/Source/Runtime/Projects/Private/PluginManager.cpp
+++ b/Engine/Source/Runtime/Projects/Private/PluginManager.cpp
@@ -267,7 +267,9 @@ bool FPluginManager::ConfigureEnabledPlugins()
 		const FProjectDescriptor *Project = IProjectManager::Get().GetCurrentProject();
 		if(Project != nullptr)
 		{
-			for(const FPluginReferenceDescriptor& Plugin: Project->Plugins)
+			// Take a copy of the Project's plugins as we may remove some
+			TArray<FPluginReferenceDescriptor> PluginsCopy = Project->Plugins;
+			for(const FPluginReferenceDescriptor& Plugin: PluginsCopy)
 			{
 				if(Plugin.bEnabled && !FindPluginInstance(Plugin.Name).IsValid())
 				{
@@ -292,6 +294,27 @@ bool FPluginManager::ConfigureEnabledPlugins()
 							return false;
 						}
 
+						FText FailReason;
+						if (!IProjectManager::Get().SetPluginEnabled(*Plugin.Name, false, FailReason))
+						{
+							FMessageDialog::Open(EAppMsgType::Ok, FailReason);
+						}
+					}
+				}
+				else if (Plugin.bEnabled && Project->Modules.Num() == 0)
+				{
+					// Content only project - check whether any plugins are incompatible and offer to disable instead of trying to build them later
+					TSharedPtr<FPlugin> PluginInstance = FindPluginInstance(Plugin.Name);
+					TArray<FString> IncompatibleFiles;
+					if (!FModuleDescriptor::CheckModuleCompatbility(PluginInstance->Descriptor.Modules, PluginInstance->LoadedFrom == EPluginLoadedFrom::GameProject, IncompatibleFiles))
+					{
+						// Ask whether to disable plugin if incompatible
+						FText Caption(LOCTEXT("IncompatiblePluginCaption", "Plugin missing or incompatible"));
+						if (FMessageDialog::Open(EAppMsgType::YesNo, FText::Format(LOCTEXT("IncompatiblePluginText", "Missing or incompatible modules in {0} plugin - would you like to disable it? You will no longer be able to open any assets created using it."), FText::FromString(Plugin.Name)), &Caption) == EAppReturnType::No)
+						{
+							return false;
+						}
+
 						FText FailReason;
 						if (!IProjectManager::Get().SetPluginEnabled(*Plugin.Name, false, FailReason))
 						{
