Changed code checking module compatibility to only run on non monolithic builds

Also fixed spelling mistake in function that made it hard to find all instances of the same thing

diff --git a/Engine/Source/Runtime/Projects/Private/ModuleDescriptor.cpp b/Engine/Source/Runtime/Projects/Private/ModuleDescriptor.cpp
--- a/Engine/Source/Runtime/Projects/Private/ModuleDescriptor.cpp
+++ b/Engine/Source/Runtime/Projects/Private/ModuleDescriptor.cpp
@@ -382,7 +382,7 @@ void FModuleDescriptor::LoadModulesForPhase(ELoadingPhase::Type LoadingPhase, co
 	}
 }
 
-bool FModuleDescriptor::CheckModuleCompatbility(const TArray<FModuleDescriptor>& Modules, bool bGameModules, TArray<FString>& OutIncompatibleFiles)
+bool FModuleDescriptor::CheckModuleCompatibility(const TArray<FModuleDescriptor>& Modules, bool bGameModules, TArray<FString>& OutIncompatibleFiles)
 {
 	bool bResult = true;
 	for(int Idx = 0; Idx < Modules.Num(); Idx++)
diff --git a/Engine/Source/Runtime/Projects/Private/PluginManager.cpp b/Engine/Source/Runtime/Projects/Private/PluginManager.cpp
--- a/Engine/Source/Runtime/Projects/Private/PluginManager.cpp
+++ b/Engine/Source/Runtime/Projects/Private/PluginManager.cpp
@@ -301,12 +301,14 @@ bool FPluginManager::ConfigureEnabledPlugins()
 						}
 					}
 				}
+#if !IS_MONOLITHIC
+				// Only check this when in a non-monolithic build where modules could be in separate binaries
 				else if (Plugin.bEnabled && Project->Modules.Num() == 0)
 				{
 					// Content only project - check whether any plugins are incompatible and offer to disable instead of trying to build them later
 					TSharedPtr<FPlugin> PluginInstance = FindPluginInstance(Plugin.Name);
 					TArray<FString> IncompatibleFiles;
-					if (!FModuleDescriptor::CheckModuleCompatbility(PluginInstance->Descriptor.Modules, PluginInstance->LoadedFrom == EPluginLoadedFrom::GameProject, IncompatibleFiles))
+					if (!FModuleDescriptor::CheckModuleCompatibility(PluginInstance->Descriptor.Modules, PluginInstance->LoadedFrom == EPluginLoadedFrom::GameProject, IncompatibleFiles))
 					{
 						// Ask whether to disable plugin if incompatible
 						FText Caption(LOCTEXT("IncompatiblePluginCaption", "Plugin missing or incompatible"));
@@ -322,6 +324,7 @@ bool FPluginManager::ConfigureEnabledPlugins()
 						}
 					}
 				}
+#endif //!IS_MONOLITHIC
 			}
 		}
 
@@ -513,7 +516,7 @@ bool FPluginManager::CheckModuleCompatibility(TArray<FString>& OutIncompatibleMo
 	for (TArray< TSharedRef< FPlugin > >::TConstIterator Iter(AllPlugins); Iter; ++Iter)
 	{
 		const TSharedRef< FPlugin > &Plugin = *Iter;
-		if (Plugin->bEnabled && !FModuleDescriptor::CheckModuleCompatbility(Plugin->Descriptor.Modules, Plugin->LoadedFrom == EPluginLoadedFrom::GameProject, OutIncompatibleModules))
+		if (Plugin->bEnabled && !FModuleDescriptor::CheckModuleCompatibility(Plugin->Descriptor.Modules, Plugin->LoadedFrom == EPluginLoadedFrom::GameProject, OutIncompatibleModules))
 		{
 			bResult = false;
 		}
diff --git a/Engine/Source/Runtime/Projects/Private/ProjectManager.cpp b/Engine/Source/Runtime/Projects/Private/ProjectManager.cpp
--- a/Engine/Source/Runtime/Projects/Private/ProjectManager.cpp
+++ b/Engine/Source/Runtime/Projects/Private/ProjectManager.cpp
@@ -98,7 +98,7 @@ bool FProjectManager::LoadModulesForProject( const ELoadingPhase::Type LoadingPh
 
 bool FProjectManager::CheckModuleCompatibility(TArray<FString>& OutIncompatibleModules)
 {
-	return !CurrentProject.IsValid() || FModuleDescriptor::CheckModuleCompatbility(CurrentProject->Modules, true, OutIncompatibleModules);
+	return !CurrentProject.IsValid() || FModuleDescriptor::CheckModuleCompatibility(CurrentProject->Modules, true, OutIncompatibleModules);
 }
 
 const FString& FProjectManager::GetAutoLoadProjectFileName()
diff --git a/Engine/Source/Runtime/Projects/Public/ModuleDescriptor.h b/Engine/Source/Runtime/Projects/Public/ModuleDescriptor.h
--- a/Engine/Source/Runtime/Projects/Public/ModuleDescriptor.h
+++ b/Engine/Source/Runtime/Projects/Public/ModuleDescriptor.h
@@ -136,5 +136,5 @@ struct PROJECTS_API FModuleDescriptor
 	static void LoadModulesForPhase(ELoadingPhase::Type LoadingPhase, const TArray<FModuleDescriptor>& Modules, TMap<FName, EModuleLoadResult>& ModuleLoadErrors);
 
 	/** Checks that all modules are compatible with the current engine version. Returns false and appends a list of names to OutIncompatibleFiles if not. */
-	static bool CheckModuleCompatbility(const TArray<FModuleDescriptor>& Modules, bool bGameModules, TArray<FString>& OutIncompatibleFiles);
+	static bool CheckModuleCompatibility(const TArray<FModuleDescriptor>& Modules, bool bGameModules, TArray<FString>& OutIncompatibleFiles);
 };
