Added a dialog box to alert users when there is a problem loading a plugin descriptor file

diff --git a/Engine/Source/Runtime/Projects/Private/PluginManager.cpp b/Engine/Source/Runtime/Projects/Private/PluginManager.cpp
--- a/Engine/Source/Runtime/Projects/Private/PluginManager.cpp
+++ b/Engine/Source/Runtime/Projects/Private/PluginManager.cpp
@@ -184,7 +184,11 @@ void FPluginManager::ReadPluginsInDirectory(const FString& PluginsDirectory, con
 			else
 			{
 				// NOTE: Even though loading of this plugin failed, we'll keep processing other plugins
-				UE_LOG(LogPluginManager, Error, TEXT("%s (%s)"), *FailureReason.ToString(), *FileName);
+				FString FullPath = FPaths::ConvertRelativePathToFull(FileName);
+				FText FailureMessage = FText::Format(LOCTEXT("FailureFormat", "{0} ({1})"), FailureReason, FText::FromString(FullPath));
+				FText DialogTitle = LOCTEXT("PluginFailureTitle", "Failed to load Plugin");
+				UE_LOG(LogPluginManager, Error, TEXT("%s"), *FailureMessage.ToString());
+				FMessageDialog::Open(EAppMsgType::Ok, FailureMessage, &DialogTitle);
 			}
 		}
 	}
