Missing files from CL 2641800.

diff --git a/Engine/Source/Runtime/Projects/Private/VersionManifest.cpp b/Engine/Source/Runtime/Projects/Private/VersionManifest.cpp
new file mode 100644
--- /dev/null
+++ b/Engine/Source/Runtime/Projects/Private/VersionManifest.cpp
@@ -0,0 +1,106 @@
+// Copyright 1998-2015 Epic Games, Inc. All Rights Reserved.
+
+#include "ProjectsPrivatePCH.h"
+#include "VersionManifest.h"
+
+FVersionManifest::FVersionManifest()
+{
+	Changelist = 0;
+}
+
+FString FVersionManifest::GetFileName(const FString& DirectoryName, bool bIsGameFolder)
+{
+	FString FileName = DirectoryName / FPlatformProcess::ExecutableName();
+	if(FApp::GetBuildConfiguration() == EBuildConfigurations::DebugGame && bIsGameFolder)
+	{
+		FileName += FString::Printf(TEXT("-%s-DebugGame"), FPlatformProcess::GetBinariesSubdirectory());
+	}
+	return FileName + TEXT(".modules");
+}
+
+bool FVersionManifest::TryRead(const FString& FileName, FVersionManifest& Manifest)
+{
+	// Read the file to a string
+	FString Text;
+	if (!FFileHelper::LoadFileToString(Text, *FileName))
+	{
+		return false;
+	}
+
+	// Deserialize a JSON object from the string
+	TSharedPtr< FJsonObject > ObjectPtr;
+	TSharedRef< TJsonReader<> > Reader = TJsonReaderFactory<>::Create(Text);
+	if (!FJsonSerializer::Deserialize(Reader, ObjectPtr) || !ObjectPtr.IsValid() )
+	{
+		return false;
+	}
+	FJsonObject& Object = *ObjectPtr.Get();
+
+	// Read the changelist and build id
+	if(!Object.TryGetNumberField(TEXT("Changelist"), Manifest.Changelist) || !Object.TryGetStringField(TEXT("BuildId"), Manifest.BuildId))
+	{
+		return false;
+	}
+
+	// Read the module mappings
+	TSharedPtr<FJsonObject> ModulesObject = Object.GetObjectField(TEXT("Modules"));
+	if(ModulesObject.IsValid())
+	{
+		for(TPair<FString, TSharedPtr<FJsonValue>>& Pair: ModulesObject->Values)
+		{
+			if(Pair.Value->Type == EJson::String)
+			{
+				Manifest.ModuleNameToFileName.FindOrAdd(Pair.Key) = Pair.Value->AsString();
+			}
+		}
+	}
+
+	return true;
+}
+
+
+
+
+
+FVersionedModuleEnumerator::FVersionedModuleEnumerator()
+{
+}
+
+bool FVersionedModuleEnumerator::RegisterWithModuleManager()
+{
+	FString VersionManifestFileName = FVersionManifest::GetFileName(FPlatformProcess::BaseDir(), false);
+	if(!FVersionManifest::TryRead(VersionManifestFileName, InitialManifest))
+	{
+		return false;
+	}
+
+	FModuleManager::Get().QueryModulesDelegate.BindRaw(this, &FVersionedModuleEnumerator::QueryModules);
+	return true;
+}
+
+const FVersionManifest& FVersionedModuleEnumerator::GetInitialManifest() const
+{
+	return InitialManifest;
+}
+
+void FVersionedModuleEnumerator::QueryModules(const FString& InDirectoryName, bool bIsGameDirectory, TMap<FString, FString>& OutModules) const
+{
+	FVersionManifest Manifest;
+	if(FVersionManifest::TryRead(FVersionManifest::GetFileName(InDirectoryName, bIsGameDirectory), Manifest) && IsMatchingVersion(Manifest))
+	{
+		OutModules = Manifest.ModuleNameToFileName;
+	}
+}
+
+bool FVersionedModuleEnumerator::IsMatchingVersion(const FVersionManifest& Manifest) const
+{
+	if(InitialManifest.Changelist == 0)
+	{
+		return Manifest.Changelist == 0 && (InitialManifest.BuildId == Manifest.BuildId || Manifest.BuildId.Len() == 0);
+	}
+	else
+	{
+		return InitialManifest.Changelist == Manifest.Changelist;
+	}
+}
+
diff --git a/Engine/Source/Runtime/Projects/Public/VersionManifest.h b/Engine/Source/Runtime/Projects/Public/VersionManifest.h
new file mode 100644
--- /dev/null
+++ b/Engine/Source/Runtime/Projects/Public/VersionManifest.h
@@ -0,0 +1,52 @@
+// Copyright 1998-2015 Epic Games, Inc. All Rights Reserved.
+
+#pragma once
+
+/**
+ * Stores a record of a built target, with all metadata that other tools may need to know about the build.
+ */
+class PROJECTS_API FVersionManifest
+{
+public:
+	uint32 Changelist;
+	FString BuildId;
+	TMap<FString, FString> ModuleNameToFileName;
+
+	/**
+	 * Default constructor 
+	 */
+	FVersionManifest();
+
+	/**
+	 * Gets the path to a version manifest for the given folder.
+	 *
+	 * @param DirectoryName		Directory to read from
+	 * @param bIsGameFolder		Whether the directory is a game folder of not. Used to adjust the name if the application is running in DebugGame.
+	 * @return The filename of the version manifest.
+	 */
+	static FString GetFileName(const FString& DirectoryName, bool bIsGameFolder);
+
+	/**
+	 * Read a version manifest from disk.
+	 *
+	 * @param FileName		Filename to read from
+	 */
+	static bool TryRead(const FString& FileName, FVersionManifest& Manifest);
+};
+
+/**
+ * Adapter for the module manager to be able to read and Enumerates the contents Stores a record of a built target, with all metadata that other tools may need to know about the build.
+ */
+class PROJECTS_API FVersionedModuleEnumerator
+{
+public:
+	FVersionedModuleEnumerator();
+	bool RegisterWithModuleManager();
+	const FVersionManifest& GetInitialManifest() const;
+
+private:
+	FVersionManifest InitialManifest;
+
+	void QueryModules(const FString& InDirectoryName, bool bIsGameDirectory, TMap<FString, FString>& OutModules) const;
+	bool IsMatchingVersion(const FVersionManifest& Manifest) const;
+};
