UE-20140 - now only add the plugin binaries directories to the path if the plugin is enabled

diff --git a/Engine/Source/Runtime/Projects/Private/PluginManager.cpp b/Engine/Source/Runtime/Projects/Private/PluginManager.cpp
--- a/Engine/Source/Runtime/Projects/Private/PluginManager.cpp
+++ b/Engine/Source/Runtime/Projects/Private/PluginManager.cpp
@@ -142,13 +142,6 @@ void FPluginManager::DiscoverAllPlugins()
 {
 	ensure( AllPlugins.Num() == 0 );		// Should not have already been initialized!
 	ReadAllPlugins(AllPlugins);
-
-	// Add the plugin binaries directory
-	for(const TSharedRef<FPlugin>& Plugin: AllPlugins)
-	{
-		const FString PluginBinariesPath = FPaths::Combine(*FPaths::GetPath(Plugin->FileName), TEXT("Binaries"), FPlatformProcess::GetBinariesSubdirectory());
-		FModuleManager::Get().AddBinariesDirectory(*PluginBinariesPath, Plugin->LoadedFrom == EPluginLoadedFrom::GameProject);
-	}
 }
 
 void FPluginManager::ReadAllPlugins(TArray<TSharedRef<FPlugin>>& Plugins)
@@ -360,6 +353,10 @@ bool FPluginManager::ConfigureEnabledPlugins()
 		{
 			if (Plugin->bEnabled)
 			{
+				// Add the plugin binaries directory
+				const FString PluginBinariesPath = FPaths::Combine(*FPaths::GetPath(Plugin->FileName), TEXT("Binaries"), FPlatformProcess::GetBinariesSubdirectory());
+				FModuleManager::Get().AddBinariesDirectory(*PluginBinariesPath, Plugin->LoadedFrom == EPluginLoadedFrom::GameProject);
+
 				// Build the list of content folders
 				if (Plugin->Descriptor.bCanContainContent)
 				{
@@ -393,7 +390,7 @@ bool FPluginManager::ConfigureEnabledPlugins()
 		IPlatformFile& PlatformFile = FPlatformFileManager::Get().GetPlatformFile();
 		for(TSharedRef<IPlugin> Plugin: GetEnabledPlugins())
 		{
-			if(Plugin->CanContainContent() && ensure(RegisterMountPointDelegate.IsBound()))
+			if (Plugin->CanContainContent() && ensure(RegisterMountPointDelegate.IsBound()))
 			{
 				FString ContentDir = Plugin->GetContentDir();
 				RegisterMountPointDelegate.Execute(Plugin->GetMountedAssetPath(), ContentDir);
