fix for failure to validate plugins specifically enabled in the uproject file

diff --git a/Engine/Source/Runtime/Projects/Private/PluginManager.cpp b/Engine/Source/Runtime/Projects/Private/PluginManager.cpp
--- a/Engine/Source/Runtime/Projects/Private/PluginManager.cpp
+++ b/Engine/Source/Runtime/Projects/Private/PluginManager.cpp
@@ -298,30 +298,6 @@ bool FPluginManager::ConfigureEnabledPlugins()
 						}
 					}
 				}
-#if !IS_MONOLITHIC
-				// Only check this when in a non-monolithic build where modules could be in separate binaries
-				else if (Plugin.bEnabled && Project->Modules.Num() == 0)
-				{
-					// Content only project - check whether any plugins are incompatible and offer to disable instead of trying to build them later
-					TSharedPtr<FPlugin> PluginInstance = FindPluginInstance(Plugin.Name);
-					TArray<FString> IncompatibleFiles;
-					if (!FModuleDescriptor::CheckModuleCompatibility(PluginInstance->Descriptor.Modules, PluginInstance->LoadedFrom == EPluginLoadedFrom::GameProject, IncompatibleFiles))
-					{
-						// Ask whether to disable plugin if incompatible
-						FText Caption(LOCTEXT("IncompatiblePluginCaption", "Plugin missing or incompatible"));
-						if (FMessageDialog::Open(EAppMsgType::YesNo, FText::Format(LOCTEXT("IncompatiblePluginText", "Missing or incompatible modules in {0} plugin - would you like to disable it? You will no longer be able to open any assets created using it."), FText::FromString(Plugin.Name)), &Caption) == EAppReturnType::No)
-						{
-							return false;
-						}
-
-						FText FailReason;
-						if (!IProjectManager::Get().SetPluginEnabled(*Plugin.Name, false, FailReason))
-						{
-							FMessageDialog::Open(EAppMsgType::Ok, FailReason);
-						}
-					}
-				}
-#endif //!IS_MONOLITHIC
 			}
 		}
 
@@ -357,7 +333,31 @@ bool FPluginManager::ConfigureEnabledPlugins()
 				const FString PluginBinariesPath = FPaths::Combine(*FPaths::GetPath(Plugin->FileName), TEXT("Binaries"), FPlatformProcess::GetBinariesSubdirectory());
 				FModuleManager::Get().AddBinariesDirectory(*PluginBinariesPath, Plugin->LoadedFrom == EPluginLoadedFrom::GameProject);
 
-				// Build the list of content folders
+#if !IS_MONOLITHIC
+				// Only check this when in a non-monolithic build where modules could be in separate binaries
+				if (Project->Modules.Num() == 0)
+				{
+					// Content only project - check whether any plugins are incompatible and offer to disable instead of trying to build them later
+					TArray<FString> IncompatibleFiles;
+					if (!FModuleDescriptor::CheckModuleCompatibility(Plugin->Descriptor.Modules, Plugin->LoadedFrom == EPluginLoadedFrom::GameProject, IncompatibleFiles))
+					{
+						// Ask whether to disable plugin if incompatible
+						FText Caption(LOCTEXT("IncompatiblePluginCaption", "Plugin missing or incompatible"));
+						if (FMessageDialog::Open(EAppMsgType::YesNo, FText::Format(LOCTEXT("IncompatiblePluginText", "Missing or incompatible modules in {0} plugin - would you like to disable it? You will no longer be able to open any assets created using it."), FText::FromString(Plugin->Name)), &Caption) == EAppReturnType::No)
+						{
+							return false;
+						}
+
+						FText FailReason;
+						if (!IProjectManager::Get().SetPluginEnabled(*Plugin->Name, false, FailReason))
+						{
+							FMessageDialog::Open(EAppMsgType::Ok, FailReason);
+						}
+					}
+				}
+#endif //!IS_MONOLITHIC
+
+			// Build the list of content folders
 				if (Plugin->Descriptor.bCanContainContent)
 				{
 					if (auto EngineConfigFile = GConfig->Find(GEngineIni, false))
