Can now point the editor to additional plugin directories using the -PLUGIN commandline param (to utilize plugins not located in the traditional places).

diff --git a/Engine/Source/Runtime/Projects/Private/PluginManager.cpp b/Engine/Source/Runtime/Projects/Private/PluginManager.cpp
--- a/Engine/Source/Runtime/Projects/Private/PluginManager.cpp
+++ b/Engine/Source/Runtime/Projects/Private/PluginManager.cpp
@@ -13,6 +13,43 @@ namespace PluginSystemDefs
 	    NOTE: This constant exists in UnrealBuildTool code as well. */
 	static const TCHAR PluginDescriptorFileExtension[] = TEXT( ".uplugin" );
 
+	/**
+	 * Parsing the command line and loads any foreign plugins that were
+	 * specified using the -PLUGIN= command.
+	 *
+	 * @param  CommandLine    The commandline used to launch the editor.
+	 * @param  SearchPathsOut 
+	 * @return The number of plugins that were specified using the -PLUGIN param.
+	 */
+	static int32 GetAdditionalPluginPaths(TSet<FString>& PluginPathsOut)
+	{
+		const TCHAR* SwitchStr = TEXT("PLUGIN=");
+		const int32  SwitchLen = FCString::Strlen(SwitchStr);
+
+		int32 PluginCount = 0;
+
+		const TCHAR* SearchStr = FCommandLine::Get();
+		do
+		{
+			FString PluginPath;
+
+			SearchStr = FCString::Strfind(SearchStr, SwitchStr);
+			if (FParse::Value(SearchStr, SwitchStr, PluginPath))
+			{
+				FString PluginDir = FPaths::GetPath(PluginPath);
+				PluginPathsOut.Add(PluginDir);
+
+				++PluginCount;
+				SearchStr += SwitchLen + PluginPath.Len();
+			}
+			else
+			{
+				break;
+			}
+		} while (SearchStr != nullptr);
+
+		return PluginCount;
+	}
 }
 
 FPlugin::FPlugin(const FString& InFileName, const FPluginDescriptor& InDescriptor, EPluginLoadedFrom InLoadedFrom)
@@ -84,9 +121,6 @@ bool FPlugin::UpdateDescriptor(const FPluginDescriptor& NewDescriptor, FText& Ou
 	return true;
 }
 
-
-
-
 FPluginManager::FPluginManager()
 	: bHaveConfiguredEnabledPlugins(false)
 	, bHaveAllRequiredPlugins(false)
@@ -94,8 +128,6 @@ FPluginManager::FPluginManager()
 	DiscoverAllPlugins();
 }
 
-
-
 FPluginManager::~FPluginManager()
 {
 	// NOTE: All plugins and modules should be cleaned up or abandoned by this point
@@ -111,7 +143,7 @@ void FPluginManager::RefreshPluginsList()
 {
 	// Read a new list of all plugins
 	TArray<TSharedRef<FPlugin>> NewPlugins;
-	ReadAllPlugins(NewPlugins);
+	ReadAllPlugins(NewPlugins, PluginDiscoveryPaths);
 
 	// Build a list of filenames for plugins which are enabled, and remove the rest
 	TArray<FString> EnabledPluginFileNames;
@@ -141,10 +173,12 @@ void FPluginManager::RefreshPluginsList()
 void FPluginManager::DiscoverAllPlugins()
 {
 	ensure( AllPlugins.Num() == 0 );		// Should not have already been initialized!
-	ReadAllPlugins(AllPlugins);
+
+	PluginSystemDefs::GetAdditionalPluginPaths(PluginDiscoveryPaths);
+	ReadAllPlugins(AllPlugins, PluginDiscoveryPaths);
 }
 
-void FPluginManager::ReadAllPlugins(TArray<TSharedRef<FPlugin>>& Plugins)
+void FPluginManager::ReadAllPlugins(TArray<TSharedRef<FPlugin>>& Plugins, const TSet<FString>& ExtraSearchPaths)
 {
 #if (WITH_ENGINE && !IS_PROGRAM) || WITH_PLUGIN_SUPPORT
 	// Find "built-in" plugins.  That is, plugins situated right within the Engine directory.
@@ -155,6 +189,11 @@ void FPluginManager::ReadAllPlugins(TArray<TSharedRef<FPlugin>>& Plugins)
 	{
 		ReadPluginsInDirectory(FPaths::GamePluginsDir(), EPluginLoadedFrom::GameProject, Plugins);
 	}
+
+	for (const FString& ExtraSearchPath : ExtraSearchPaths)
+	{
+		ReadPluginsInDirectory(ExtraSearchPath, EPluginLoadedFrom::GameProject, Plugins);
+	}
 #endif
 }
 
@@ -594,4 +633,13 @@ TArray< FPluginStatus > FPluginManager::QueryStatusForAllPlugins() const
 	return PluginStatuses;
 }
 
+void FPluginManager::AddPluginSearchPath(const FString& ExtraDiscoveryPath, bool bRefresh)
+{
+	PluginDiscoveryPaths.Add(ExtraDiscoveryPath);
+	if (bRefresh)
+	{
+		RefreshPluginsList();
+	}
+}
+
 #undef LOCTEXT_NAMESPACE
diff --git a/Engine/Source/Runtime/Projects/Private/PluginManager.h b/Engine/Source/Runtime/Projects/Private/PluginManager.h
--- a/Engine/Source/Runtime/Projects/Private/PluginManager.h
+++ b/Engine/Source/Runtime/Projects/Private/PluginManager.h
@@ -70,6 +70,7 @@ public:
 	virtual TArray<TSharedRef<IPlugin>> GetEnabledPlugins() override;
 	virtual TArray<TSharedRef<IPlugin>> GetDiscoveredPlugins() override;
 	virtual TArray< FPluginStatus > QueryStatusForAllPlugins() const override;
+	virtual void AddPluginSearchPath(const FString& ExtraDiscoveryPath, bool bRefresh = true) override;
 
 private:
 
@@ -78,7 +79,7 @@ private:
 	void DiscoverAllPlugins();
 
 	/** Reads all the plugin descriptors */
-	static void ReadAllPlugins(TArray<TSharedRef<FPlugin>>& Plugins);
+	static void ReadAllPlugins(TArray<TSharedRef<FPlugin>>& Plugins, const TSet<FString>& ExtraSearchPaths);
 
 	/** Reads all the plugin descriptors from disk */
 	static void ReadPluginsInDirectory(const FString& PluginsDirectory, const EPluginLoadedFrom LoadedFrom, TArray<TSharedRef<FPlugin>>& Plugins);
@@ -105,6 +106,9 @@ private:
 
 	/** Set if all the required plugins are available */
 	bool bHaveAllRequiredPlugins;
+
+	/** List of additional directory paths to search for plugins within */
+	TSet<FString> PluginDiscoveryPaths;
 };
 
 
diff --git a/Engine/Source/Runtime/Projects/Public/Interfaces/IPluginManager.h b/Engine/Source/Runtime/Projects/Public/Interfaces/IPluginManager.h
--- a/Engine/Source/Runtime/Projects/Public/Interfaces/IPluginManager.h
+++ b/Engine/Source/Runtime/Projects/Public/Interfaces/IPluginManager.h
@@ -196,6 +196,16 @@ public:
 	 */
 	virtual TArray<FPluginStatus> QueryStatusForAllPlugins() const = 0;
 
+	/**
+	 * Stores the specified path, utilizing it in future search passes when 
+	 * searching for available plugins. Optionally refreshes the manager after 
+	 * the new path has been added.
+	 * 
+	 * @param  ExtraDiscoveryPath	The path you want searched for additional plugins.
+	 * @param  bRefresh				Signals the function to refresh the plugin database after the new path has been added
+	 */
+	virtual void AddPluginSearchPath(const FString& ExtraDiscoveryPath, bool bRefresh = true) = 0;
+
 public:
 
 	/**
