Remove debug code that writes a copy of the project descriptor to disk.

diff --git a/Engine/Source/Runtime/Projects/Private/ProjectManager.cpp b/Engine/Source/Runtime/Projects/Private/ProjectManager.cpp
--- a/Engine/Source/Runtime/Projects/Private/ProjectManager.cpp
+++ b/Engine/Source/Runtime/Projects/Private/ProjectManager.cpp
@@ -23,10 +23,6 @@ bool FProjectManager::LoadProjectFile( const FString& InProjectFile )
 	TSharedPtr<FProjectDescriptor> Descriptor = MakeShareable(new FProjectDescriptor());
 	if(Descriptor->Load(InProjectFile, FailureReason))
 	{
-		Descriptor->Save("D:\\P4 Streams\\Templates\\TP_Flying\\Out.json", FailureReason);
-
-
-
 		// Create the project
 		CurrentProject = Descriptor;
 		return true;
