Plugins - No longer creating plugin descriptors using the full file name as that breaks cooking, instead doing it where the UAT call is for packaging a plugin.

diff --git a/Engine/Plugins/Editor/PluginBrowser/Source/PluginBrowser/Private/SPluginTile.cpp b/Engine/Plugins/Editor/PluginBrowser/Source/PluginBrowser/Private/SPluginTile.cpp
--- a/Engine/Plugins/Editor/PluginBrowser/Source/PluginBrowser/Private/SPluginTile.cpp
+++ b/Engine/Plugins/Editor/PluginBrowser/Source/PluginBrowser/Private/SPluginTile.cpp
@@ -555,8 +555,9 @@ void SPluginTile::OnPackagePlugin()
 	}
 
 	// Ensure path is full rather than relative (for macs)
-	FString UPluginFilename = Plugin->GetDescriptorFileName();
-	FString CommandLine = FString::Printf(TEXT("BuildPlugin -Rocket -Plugin=\"%s\" -Package=\"%s\""), *UPluginFilename, *OutputDirectory);
+	FString DescriptorFilename = Plugin->GetDescriptorFileName();
+	FString DescriptorFullPath = FPaths::ConvertRelativePathToFull(DescriptorFilename);
+	FString CommandLine = FString::Printf(TEXT("BuildPlugin -Rocket -Plugin=\"%s\" -Package=\"%s\""), *DescriptorFullPath, *OutputDirectory);
 
 #if PLATFORM_WINDOWS
 	FText PlatformName = LOCTEXT("PlatformName_Windows", "Windows");
diff --git a/Engine/Source/Runtime/Projects/Private/PluginManager.cpp b/Engine/Source/Runtime/Projects/Private/PluginManager.cpp
--- a/Engine/Source/Runtime/Projects/Private/PluginManager.cpp
+++ b/Engine/Source/Runtime/Projects/Private/PluginManager.cpp
@@ -67,6 +67,7 @@ FPlugin::FPlugin(const FString& InFileName, const FPluginDescriptor& InDescripto
 	, LoadedFrom(InLoadedFrom)
 	, bEnabled(false)
 {
+
 }
 
 FPlugin::~FPlugin()
@@ -215,17 +216,21 @@ void FPluginManager::ReadPluginsInDirectory(const FString& PluginsDirectory, con
 
 		for(const FString& FileName: FileNames)
 		{
-			FString FullPath = FPaths::ConvertRelativePathToFull(FileName);
-
 			FPluginDescriptor Descriptor;
 			FText FailureReason;
-			if ( Descriptor.Load(FullPath, FailureReason) )
+			if ( Descriptor.Load(FileName, FailureReason) )
 			{
-				Plugins.Add(MakeShareable(new FPlugin(FullPath, Descriptor, LoadedFrom)));
+				TSharedRef<FPlugin> Plugin = MakeShareable(new FPlugin(FileName, Descriptor, LoadedFrom));
+				
+				FString FullPath = FPaths::ConvertRelativePathToFull(FileName);
+				UE_LOG(LogPluginManager, Log, TEXT("Loaded Plugin %s, From %s"), *Plugin->GetName(), *FullPath);
+
+				Plugins.Add(Plugin);
 			}
 			else
 			{
 				// NOTE: Even though loading of this plugin failed, we'll keep processing other plugins
+				FString FullPath = FPaths::ConvertRelativePathToFull(FileName);
 				FText FailureMessage = FText::Format(LOCTEXT("FailureFormat", "{0} ({1})"), FailureReason, FText::FromString(FullPath));
 				FText DialogTitle = LOCTEXT("PluginFailureTitle", "Failed to load Plugin");
 				UE_LOG(LogPluginManager, Error, TEXT("%s"), *FailureMessage.ToString());
