Fix localization gathering warnings
#rnx

diff --git a/Engine/Source/Editor/AnimGraph/Private/AnimGraphNode_BoneDrivenController.cpp b/Engine/Source/Editor/AnimGraph/Private/AnimGraphNode_BoneDrivenController.cpp
--- a/Engine/Source/Editor/AnimGraph/Private/AnimGraphNode_BoneDrivenController.cpp
+++ b/Engine/Source/Editor/AnimGraph/Private/AnimGraphNode_BoneDrivenController.cpp
@@ -193,19 +193,19 @@ void UAnimGraphNode_BoneDrivenController::ValidateAnimNodeDuringCompilation(USke
 {
 	if (ForSkeleton->GetReferenceSkeleton().FindBoneIndex(Node.SourceBone.BoneName) == INDEX_NONE)
 	{
-		MessageLog.Warning(*LOCTEXT("NoSourceBone", "@@ - You must pick a source bone as the Driver joint").ToString(), this);
+		MessageLog.Warning(*LOCTEXT("DriverJoint_NoSourceBone", "@@ - You must pick a source bone as the Driver joint").ToString(), this);
 	}
 
 	if (Node.SourceComponent == EComponentType::None)
 	{
-		MessageLog.Warning(*LOCTEXT("NoSourceComponent", "@@ - You must pick a source component on the Driver joint").ToString(), this);
+		MessageLog.Warning(*LOCTEXT("DriverJoint_NoSourceComponent", "@@ - You must pick a source component on the Driver joint").ToString(), this);
 	}
 	
 	if (Node.DestinationMode == EDrivenDestinationMode::Bone)
 	{
 		if (ForSkeleton->GetReferenceSkeleton().FindBoneIndex(Node.TargetBone.BoneName) == INDEX_NONE)
 		{
-			MessageLog.Warning(*LOCTEXT("NoTargetBone", "@@ - You must pick a target bone as the Driven joint").ToString(), this);
+			MessageLog.Warning(*LOCTEXT("DriverJoint_NoTargetBone", "@@ - You must pick a target bone as the Driven joint").ToString(), this);
 		}
 		
 		const bool bAffectsTranslation = Node.bAffectTargetTranslationX || Node.bAffectTargetTranslationY || Node.bAffectTargetTranslationZ;
@@ -214,7 +214,7 @@ void UAnimGraphNode_BoneDrivenController::ValidateAnimNodeDuringCompilation(USke
 
 		if (!bAffectsTranslation && !bAffectsRotation && !bAffectsScale)
 		{
-			MessageLog.Warning(*LOCTEXT("NoTargetComponent", "@@ - You must pick one or more target components on the Driven joint").ToString(), this);
+			MessageLog.Warning(*LOCTEXT("DriverJoint_NoTargetComponent", "@@ - You must pick one or more target components on the Driven joint").ToString(), this);
 		}
 	}
 
diff --git a/Engine/Source/Editor/AnimGraph/Private/AnimGraphNode_PoseDriver.cpp b/Engine/Source/Editor/AnimGraph/Private/AnimGraphNode_PoseDriver.cpp
--- a/Engine/Source/Editor/AnimGraph/Private/AnimGraphNode_PoseDriver.cpp
+++ b/Engine/Source/Editor/AnimGraph/Private/AnimGraphNode_PoseDriver.cpp
@@ -64,7 +64,7 @@ void UAnimGraphNode_PoseDriver::ValidateAnimNodeDuringCompilation(USkeleton* For
 {
 	if (Node.SourceBones.Num() == 0)
 	{
-		MessageLog.Warning(*LOCTEXT("NoSourceBone", "You must specify at least one Source Bone").ToString(), this);
+		MessageLog.Warning(*LOCTEXT("PoseDriver_NoSourceBone", "You must specify at least one Source Bone").ToString(), this);
 	}
 
 	FName MissingBoneName = NAME_None;
diff --git a/Engine/Source/Editor/BlueprintGraph/Private/EdGraphSchema_K2.cpp b/Engine/Source/Editor/BlueprintGraph/Private/EdGraphSchema_K2.cpp
--- a/Engine/Source/Editor/BlueprintGraph/Private/EdGraphSchema_K2.cpp
+++ b/Engine/Source/Editor/BlueprintGraph/Private/EdGraphSchema_K2.cpp
@@ -3807,7 +3807,7 @@ void UEdGraphSchema_K2::GetVariableTypeTree(TArray< TSharedPtr<FPinTypeTreeInfo>
 		}
 		else
 		{
-			TypeTree.Add(MakeShareable(new FPinTypeTreeInfo(GetCategoryText(PC_Object, true), PC_Object, this, LOCTEXT("ObjectType", "Hard reference to an Object"), true, TypesDatabasePtr)));
+			TypeTree.Add(MakeShareable(new FPinTypeTreeInfo(GetCategoryText(PC_Object, true), PC_Object, this, LOCTEXT("ObjectTypeHardReference", "Hard reference to an Object"), true, TypesDatabasePtr)));
 			TypeTree.Add(MakeShareable(new FPinTypeTreeInfo(GetCategoryText(PC_Class, true), PC_Class, this, LOCTEXT("ClassType", "Hard reference to a Class"), true, TypesDatabasePtr)));
 			TypeTree.Add(MakeShareable(new FPinTypeTreeInfo(GetCategoryText(PC_Asset, true), PC_Asset, this, LOCTEXT("AssetType", "Soft reference to an Object"), true, TypesDatabasePtr)));
 			TypeTree.Add(MakeShareable(new FPinTypeTreeInfo(GetCategoryText(PC_AssetClass, true), PC_AssetClass, this, LOCTEXT("AssetClassType", "Soft reference to a Class"), true, TypesDatabasePtr)));
diff --git a/Engine/Source/Runtime/Projects/Private/PluginManager.cpp b/Engine/Source/Runtime/Projects/Private/PluginManager.cpp
--- a/Engine/Source/Runtime/Projects/Private/PluginManager.cpp
+++ b/Engine/Source/Runtime/Projects/Private/PluginManager.cpp
@@ -681,11 +681,11 @@ bool FPluginManager::PromptToDisableMissingPlugin(const FString& PluginName, con
 	FText Message;
 	if (PluginName == MissingPluginName)
 	{
-		Message = FText::Format(LOCTEXT("DisablePluginMessage", "This project requires the '{0}' plugin, which could not be found.\n\nWould you like to disable it? You will no longer be able to open any assets created using it."), FText::FromString(PluginName));
+		Message = FText::Format(LOCTEXT("DisablePluginMessage_NotFound", "This project requires the '{0}' plugin, which could not be found.\n\nWould you like to disable it? You will no longer be able to open any assets created using it."), FText::FromString(PluginName));
 	}
 	else
 	{
-		Message = FText::Format(LOCTEXT("DisablePluginMessage", "This project requires the '{0}' plugin, which has a missing dependency on the '{1}' plugin.\n\nWould you like to disable it? You will no longer be able to open any assets created using it."), FText::FromString(PluginName), FText::FromString(MissingPluginName));
+		Message = FText::Format(LOCTEXT("DisablePluginMessage_MissingDependency", "This project requires the '{0}' plugin, which has a missing dependency on the '{1}' plugin.\n\nWould you like to disable it? You will no longer be able to open any assets created using it."), FText::FromString(PluginName), FText::FromString(MissingPluginName));
 	}
 
 	FText Caption(LOCTEXT("DisablePluginCaption", "Missing Plugin"));
@@ -697,11 +697,11 @@ bool FPluginManager::PromptToDisableIncompatiblePlugin(const FString& PluginName
 	FText Message;
 	if (PluginName == IncompatiblePluginName)
 	{
-		Message = FText::Format(LOCTEXT("DisablePluginMessage", "This project requires the '{0}' plugin, which is not compatible with the current engine version.\n\nWould you like to disable it? You will no longer be able to open any assets created using it."), FText::FromString(PluginName));
+		Message = FText::Format(LOCTEXT("DisablePluginMessage_IncompatibleEngineVersion", "This project requires the '{0}' plugin, which is not compatible with the current engine version.\n\nWould you like to disable it? You will no longer be able to open any assets created using it."), FText::FromString(PluginName));
 	}
 	else
 	{
-		Message = FText::Format(LOCTEXT("DisablePluginMessage", "This project requires the '{0}' plugin, which has a dependency on the '{1}' plugin, which is not compatible with the current engine version.\n\nWould you like to disable it? You will no longer be able to open any assets created using it."), FText::FromString(PluginName), FText::FromString(IncompatiblePluginName));
+		Message = FText::Format(LOCTEXT("DisablePluginMessage_IncompatibleDependency", "This project requires the '{0}' plugin, which has a dependency on the '{1}' plugin, which is not compatible with the current engine version.\n\nWould you like to disable it? You will no longer be able to open any assets created using it."), FText::FromString(PluginName), FText::FromString(IncompatiblePluginName));
 	}
 
 	FText Caption(LOCTEXT("DisablePluginCaption", "Missing Plugin"));
