Fix CIS warning.

diff --git a/Engine/Source/Runtime/Projects/Private/ProjectDescriptor.cpp b/Engine/Source/Runtime/Projects/Private/ProjectDescriptor.cpp
--- a/Engine/Source/Runtime/Projects/Private/ProjectDescriptor.cpp
+++ b/Engine/Source/Runtime/Projects/Private/ProjectDescriptor.cpp
@@ -145,7 +145,7 @@ bool FProjectDescriptor::Read(const FJsonObject& Object, const FString& PathToPr
 		if (FPlatformProperties::RequiresCookedData() && AdditionalPluginDirectoriesValue->Num() > 0)
 		{
 			AdditionalPluginDirectories.Empty();
-			FString RemappedDir = FPaths::GameDir() + TEXT("../RemappedPlugins/");
+			FString RemappedDir = FPaths::ProjectDir() + TEXT("../RemappedPlugins/");
 			AddPluginDirectory(RemappedDir);
 		}
 	}
