PluginManager was computing config file paths in a way that could diverge from what happens in ConfigCacheIni for global ini files. The end result was that the paths were resulting in two different entries for the same Ini name. The later entry was incorrect and empty and clobbered over the previous entry.

Example: Engine is installed in C:\Epic\Engine. Project is installed in C:\Epic\MyProject.. works. Project is installed in D:\Folder\MyProject .. works. Project is installed in C:\Folder\MyProject.. fails.

#tests checked Niagara works properly in projects in multiple locations, including ones that failed previously

diff --git a/Engine/Source/Runtime/Projects/Private/PluginManager.cpp b/Engine/Source/Runtime/Projects/Private/PluginManager.cpp
--- a/Engine/Source/Runtime/Projects/Private/PluginManager.cpp
+++ b/Engine/Source/Runtime/Projects/Private/PluginManager.cpp
@@ -577,6 +577,7 @@ bool FPluginManager::ConfigureEnabledPlugins()
 				}
 
 				FString PluginConfigFilename = FString::Printf(TEXT("%s%s/%s.ini"), *FPaths::GeneratedConfigDir(), ANSI_TO_TCHAR(FPlatformProperties::PlatformName()), *Plugin.Name);
+				FPaths::MakeStandardFilename(PluginConfigFilename); // This needs to match what we do in ConfigCacheIni.cpp's GetDestIniFilename method. Otherwise, the hash results will differ and the plugin's version will be overwritten later.
 				FConfigFile& PluginConfig = GConfig->Add(PluginConfigFilename, FConfigFile());
 
 				// This will write out an ini to PluginConfigFilename
