Fixed logic issue in FProjectManager::IsNonDefaultPluginEnabled

diff --git a/Engine/Source/Runtime/Projects/Private/ProjectManager.cpp b/Engine/Source/Runtime/Projects/Private/ProjectManager.cpp
--- a/Engine/Source/Runtime/Projects/Private/ProjectManager.cpp
+++ b/Engine/Source/Runtime/Projects/Private/ProjectManager.cpp
@@ -44,14 +44,14 @@ bool FProjectManager::LoadProjectFile( const FString& InProjectFile )
 		CurrentProjectModuleContextInfos.Reset();
 		return true;
 	}
-	
+
 #if PLATFORM_IOS
 	FString UpdatedMessage = FString::Printf(TEXT("%s\n%s"), *FailureReason.ToString(), TEXT("For troubleshooting, please go to https://docs.unrealengine.com/latest/INT/Platforms/iOS/GettingStarted/index.html"));
 	FailureReason = FText::FromString(UpdatedMessage);
 #endif
 	UE_LOG(LogProjectManager, Error, TEXT("%s"), *FailureReason.ToString());
 	FMessageDialog::Open(EAppMsgType::Ok, FailureReason);
-	
+
 	return false;
 }
 
@@ -94,7 +94,7 @@ bool FProjectManager::LoadModulesForProject( const ELoadingPhase::Type LoadingPh
 					{
 						FailureMessage = FText::Format( LOCTEXT("PrimaryGameModuleCouldntBeLoaded", "The game module '{0}' could not be loaded. There may be an operating system error or the module may not be properly set up."), TextModuleName );
 					}
-					else 
+					else
 					{
 						ensure(0);	// If this goes off, the error handling code should be updated for the new enum values!
 						FailureMessage = FText::Format( LOCTEXT("PrimaryGameModuleGenericLoadFailure", "The game module '{0}' failed to load for an unspecified reason.  Please report this error."), TextModuleName );
@@ -271,18 +271,24 @@ bool FProjectManager::IsNonDefaultPluginEnabled() const
 	// Check whether the setting for any default plugin has been changed
 	for(const TSharedRef<IPlugin>& Plugin: IPluginManager::Get().GetDiscoveredPlugins())
 	{
-		bool bEnabled = Plugin->IsEnabledByDefault();
-
-		bool* EnabledPtr = ConfiguredPlugins.Find(Plugin->GetName());
-		if(EnabledPtr != nullptr)
+		if (Plugin->GetDescriptor().SupportsTargetPlatform(FPlatformMisc::GetUBTPlatform()))
 		{
-			bEnabled = *EnabledPtr;
-		}
+			bool bEnabled = Plugin->IsEnabledByDefault();
 
-		bool bEnabledInDefaultExe = (Plugin->GetLoadedFrom() == EPluginLoadedFrom::Engine && Plugin->IsEnabledByDefault() && !Plugin->GetDescriptor().bInstalled);
-		if(bEnabled != bEnabledInDefaultExe)
-		{
-			return true;
+			bool* EnabledPtr = ConfiguredPlugins.Find(Plugin->GetName());
+			if(EnabledPtr != nullptr)
+			{
+				bEnabled = *EnabledPtr;
+			}
+
+			if (bEnabled)
+			{
+				bool bEnabledInDefaultExe = (Plugin->GetLoadedFrom() == EPluginLoadedFrom::Engine && Plugin->IsEnabledByDefault() && !Plugin->GetDescriptor().bInstalled);
+				if(bEnabled != bEnabledInDefaultExe)
+				{
+					return true;
+				}
+			}
 		}
 	}
 
