Fix typo.

#rnx

diff --git a/Engine/Source/Runtime/Projects/Private/PluginManager.cpp b/Engine/Source/Runtime/Projects/Private/PluginManager.cpp
--- a/Engine/Source/Runtime/Projects/Private/PluginManager.cpp
+++ b/Engine/Source/Runtime/Projects/Private/PluginManager.cpp
@@ -467,7 +467,7 @@ bool FPluginManager::ConfigureEnabledPlugins()
 				}
 			}
 
-			// Configure the plugins that were enabled from the target file
+			// Configure the plugins that were disabled from the target file
 			TArray<FString> TargetDisabledPlugins = { UBT_TARGET_DISABLED_PLUGINS };
 			for (const FString& TargetDisabledPlugin : TargetDisabledPlugins)
 			{
