Improve messaging about disabling a missing plugin.

#rnx

diff --git a/Engine/Source/Runtime/Projects/Private/PluginManager.cpp b/Engine/Source/Runtime/Projects/Private/PluginManager.cpp
--- a/Engine/Source/Runtime/Projects/Private/PluginManager.cpp
+++ b/Engine/Source/Runtime/Projects/Private/PluginManager.cpp
@@ -817,11 +817,11 @@ bool FPluginManager::PromptToDisableMissingPlugin(const FString& PluginName, con
 	FText Message;
 	if (PluginName == MissingPluginName)
 	{
-		Message = FText::Format(LOCTEXT("DisablePluginMessage_NotFound", "This project requires the '{0}' plugin, which could not be found.\n\nWould you like to disable it? You will no longer be able to open any assets created using it."), FText::FromString(PluginName));
+		Message = FText::Format(LOCTEXT("DisablePluginMessage_NotFound", "This project requires the '{0}' plugin, which could not be found. Would you like to disable it and continue?\n\nIf you do, you will no longer be able to open any assets created with it. If not, the application will close."), FText::FromString(PluginName));
 	}
 	else
 	{
-		Message = FText::Format(LOCTEXT("DisablePluginMessage_MissingDependency", "This project requires the '{0}' plugin, which has a missing dependency on the '{1}' plugin.\n\nWould you like to disable it? You will no longer be able to open any assets created using it."), FText::FromString(PluginName), FText::FromString(MissingPluginName));
+		Message = FText::Format(LOCTEXT("DisablePluginMessage_MissingDependency", "This project requires the '{0}' plugin, which has a missing dependency on the '{1}' plugin.\n\nWould you like to disable it?\n\nIf you do, you will no longer be able to open any assets created with it. If not, the application will close."), FText::FromString(PluginName), FText::FromString(MissingPluginName));
 	}
 
 	FText Caption(LOCTEXT("DisablePluginCaption", "Missing Plugin"));
