Allow you to enable certain plugins via the command line using -EnablePlugins

diff --git a/Engine/Source/Runtime/Projects/Private/PluginManager.cpp b/Engine/Source/Runtime/Projects/Private/PluginManager.cpp
--- a/Engine/Source/Runtime/Projects/Private/PluginManager.cpp
+++ b/Engine/Source/Runtime/Projects/Private/PluginManager.cpp
@@ -429,25 +429,40 @@ bool FPluginManager::ConfigureEnabledPlugins()
 		// Keep a set of all the plugin names that have been configured. We read configuration data from different places, but only configure a plugin from the first place that it's referenced.
 		TSet<FString> ConfiguredPluginNames;
 
-		// Check if we want to enable all plugins
-		if (FParse::Param(FCommandLine::Get(), TEXT("EnableAllPlugins")))
+		// Check which plugins have been enabled or excluded via the command line
 		{
-			TArray<FString> ExceptPlugins;
+			auto ParsePluginsList = [](const TCHAR* InListKey) -> TArray<FString>
 			{
-				FString ExceptPluginsStr;
-				FParse::Value(FCommandLine::Get(), TEXT("ExceptPlugins="), ExceptPluginsStr, false);
-				ExceptPluginsStr.ParseIntoArray(ExceptPlugins, TEXT(","));
+				TArray<FString> PluginsList;
+				FString PluginsListStr;
+				FParse::Value(FCommandLine::Get(), InListKey, PluginsListStr, false);
+				PluginsListStr.ParseIntoArray(PluginsList, TEXT(","));
+				return PluginsList;
+			};
+
+			// Which extra plugins should be enabled?
+			TArray<FString> ExtraPluginsToEnable;
+			if (FParse::Param(FCommandLine::Get(), TEXT("EnableAllPlugins")))
+			{
+				AllPlugins.GenerateKeyArray(ExtraPluginsToEnable);
 			}
-
-			for (const TPair<FString, TSharedRef<FPlugin>>& PluginPair : AllPlugins)
+			else
+			{
+				ExtraPluginsToEnable = ParsePluginsList(TEXT("EnablePlugins="));
+			}
+			if (ExtraPluginsToEnable.Num() > 0)
 			{
-				if (!ConfiguredPluginNames.Contains(PluginPair.Key) && !ExceptPlugins.Contains(PluginPair.Key))
+				const TArray<FString> ExceptPlugins = ParsePluginsList(TEXT("ExceptPlugins="));
+				for (const FString& EnablePluginName : ExtraPluginsToEnable)
 				{
-					if (!ConfigureEnabledPlugin(FPluginReferenceDescriptor(PluginPair.Key, true), EnabledPluginNames))
+					if (!ConfiguredPluginNames.Contains(EnablePluginName) && !ExceptPlugins.Contains(EnablePluginName))
 					{
-						return false;
+						if (!ConfigureEnabledPlugin(FPluginReferenceDescriptor(EnablePluginName, true), EnabledPluginNames))
+						{
+							return false;
+						}
+						ConfiguredPluginNames.Add(EnablePluginName);
 					}
-					ConfiguredPluginNames.Add(PluginPair.Key);
 				}
 			}
 		}
