Ensure that a project is converted to a code project if it has a plugin enabled for the target platform being packaged.

diff --git a/Engine/Source/Editor/GameProjectGeneration/Private/GameProjectUtils.cpp b/Engine/Source/Editor/GameProjectGeneration/Private/GameProjectUtils.cpp
--- a/Engine/Source/Editor/GameProjectGeneration/Private/GameProjectUtils.cpp
+++ b/Engine/Source/Editor/GameProjectGeneration/Private/GameProjectUtils.cpp
@@ -3543,6 +3543,8 @@ bool GameProjectUtils::ProjectHasCodeFiles()
 	return FileNames.Num() > 0;
 }
 
+PROJECTS_API bool HasDefaultPluginSettings(const FString& Platform);
+
 bool GameProjectUtils::ProjectRequiresBuild(const FName InPlatformInfoName)
 {
 	//  early out on projects with code files
@@ -3560,7 +3562,9 @@ bool GameProjectUtils::ProjectRequiresBuild(const FName InPlatformInfoName)
 	}
 
 	// check to see if any plugins beyond the defaults have been enabled
-	bRequiresBuild |= !IProjectManager::Get().HasDefaultPluginSettings();
+	const PlatformInfo::FPlatformInfo* PlatformInfo = PlatformInfo::FindPlatformInfo(InPlatformInfoName);
+	FName PlatformName = (PlatformInfo != nullptr) ? PlatformInfo->VanillaPlatformName : InPlatformInfoName;
+	bRequiresBuild |= !HasDefaultPluginSettings(PlatformName.ToString());
 
 	// check to see if Blueprint nativization is enabled in the Project settings
 	bRequiresBuild |= GetDefault<UProjectPackagingSettings>()->BlueprintNativizationMethod != EProjectPackagingBlueprintNativizationMethod::Disabled;
diff --git a/Engine/Source/Runtime/Projects/Private/ProjectManager.cpp b/Engine/Source/Runtime/Projects/Private/ProjectManager.cpp
--- a/Engine/Source/Runtime/Projects/Private/ProjectManager.cpp
+++ b/Engine/Source/Runtime/Projects/Private/ProjectManager.cpp
@@ -249,22 +249,24 @@ void FProjectManager::ClearSupportedTargetPlatformsForCurrentProject()
 	OnTargetPlatformsForCurrentProjectChangedEvent.Broadcast();
 }
 
-bool FProjectManager::HasDefaultPluginSettings() const
+PROJECTS_API bool HasDefaultPluginSettings(const FString& Platform)
 {
+	const FProjectDescriptor* CurrentProject = IProjectManager::Get().GetCurrentProject();
+
 	// Get settings for the plugins which are currently enabled or disabled by the project file
 	TMap<FString, bool> ConfiguredPlugins;
-	if (CurrentProject.IsValid())
+	if (CurrentProject != nullptr)
 	{
 		for (const FPluginReferenceDescriptor& PluginReference : CurrentProject->Plugins)
 		{
-			ConfiguredPlugins.Add(PluginReference.Name, PluginReference.bEnabled);
+			ConfiguredPlugins.Add(PluginReference.Name, PluginReference.IsEnabledForPlatform(Platform));
 		}
 	}
 
 	// Check whether the setting for any default plugin has been changed
 	for(const TSharedRef<IPlugin>& Plugin: IPluginManager::Get().GetDiscoveredPlugins())
 	{
-		if (Plugin->GetDescriptor().SupportsTargetPlatform(FPlatformMisc::GetUBTPlatform()))
+		if (Plugin->GetDescriptor().SupportsTargetPlatform(Platform))
 		{
 			bool bEnabled = Plugin->IsEnabledByDefault();
 
@@ -285,6 +287,11 @@ bool FProjectManager::HasDefaultPluginSettings() const
 	return true;
 }
 
+bool FProjectManager::HasDefaultPluginSettings() const
+{
+	return ::HasDefaultPluginSettings(FPlatformMisc::GetUBTPlatform());
+}
+
 bool FProjectManager::SetPluginEnabled(const FString& PluginName, bool bEnabled, FText& OutFailReason)
 {
 	// Don't go any further if there's no project loaded
