Only load embedded JSBridge modules for whitelisted platforms and remove for shipping builds in some instances.

diff --git a/Engine/Source/Runtime/Projects/Private/ModuleDescriptor.cpp b/Engine/Source/Runtime/Projects/Private/ModuleDescriptor.cpp
--- a/Engine/Source/Runtime/Projects/Private/ModuleDescriptor.cpp
+++ b/Engine/Source/Runtime/Projects/Private/ModuleDescriptor.cpp
@@ -515,6 +515,17 @@ bool FModuleDescriptor::IsLoadedInCurrentConfiguration() const
 		return false;
 	}
 
+	// If we are generating a shipping build we should not load the Embedded plugins unless the platform is iOS/Android, otherwise we rely on default handling
+	if (WhitelistPlatforms.Num() > 0 && WhitelistPlatforms.Contains(TEXT("Embedded")))
+	{
+		#if (UE_BUILD_SHIPPING)
+			if (FString(FPlatformMisc::GetUBTPlatform()) != TEXT("iOS") && FString(FPlatformMisc::GetUBTPlatform()) != TEXT("Android"))
+			{
+				return false;
+			}
+		#endif /* UE_BUILD_SHIPPING */
+	}
+
 	// Check that the runtime environment allows it to be loaded
 	switch (Type)
 	{
