Fix Plugin config stomping command line INI overrides.

diff --git a/Engine/Source/Runtime/Projects/Private/PluginManager.cpp b/Engine/Source/Runtime/Projects/Private/PluginManager.cpp
--- a/Engine/Source/Runtime/Projects/Private/PluginManager.cpp
+++ b/Engine/Source/Runtime/Projects/Private/PluginManager.cpp
@@ -618,6 +618,11 @@ bool FPluginManager::ConfigureEnabledPlugins()
 								// if plugin config overrides are applied then don't save
 								FoundConfig->NoSave = true;
 							}
+
+#if ALLOW_INI_OVERRIDE_FROM_COMMANDLINE
+							// Don't allow plugins to stomp command line overrides
+							FConfigFile::OverrideFromCommandline(FoundConfig, PluginConfigFilename);
+#endif // ALLOW_INI_OVERRIDE_FROM_COMMANDLINE
 						}
 					}
 				}
