Various fixes/improvements to program plugins.

* Runtime plugins are no longer loaded for every program (as was the case in 4.20)
* If a module is whitelisted for a specific program, it will be enabled regardless of whether the Type is  "Program" or not.
* The bCanBeUsedWithUnrealHeaderTool flag in the plugin descriptor is now translated into support for the "UnrealHeaderTool" program on load.
* Restored support for the RuntimeAndProgram module type.
* Plugins are now enabled in UHT by passing the project on the command line, rather than passing explicit plugins to enable.

diff --git a/Engine/Source/Programs/UnrealBuildTool/Configuration/UEBuildTarget.cs b/Engine/Source/Programs/UnrealBuildTool/Configuration/UEBuildTarget.cs
--- a/Engine/Source/Programs/UnrealBuildTool/Configuration/UEBuildTarget.cs
+++ b/Engine/Source/Programs/UnrealBuildTool/Configuration/UEBuildTarget.cs
@@ -3257,18 +3257,18 @@ namespace UnrealBuildTool
 				{
 					if(!Rules.EnablePlugins.Contains(PluginReference.Name, StringComparer.InvariantCultureIgnoreCase) && !Rules.DisablePlugins.Contains(PluginReference.Name, StringComparer.InvariantCultureIgnoreCase))
 					{
-					// Make sure we don't have multiple references to the same plugin
-					if(!ReferencedNames.Add(PluginReference.Name))
-					{
-						Log.TraceWarning("Plugin '{0}' is listed multiple times in project file '{1}'.", PluginReference.Name, ProjectFile);
-					}
-					else
-					{
-						AddPlugin(PluginReference, ProjectReferenceChain, ExcludeFolders, NameToInstance, NameToInfo);
+						// Make sure we don't have multiple references to the same plugin
+						if(!ReferencedNames.Add(PluginReference.Name))
+						{
+							Log.TraceWarning("Plugin '{0}' is listed multiple times in project file '{1}'.", PluginReference.Name, ProjectFile);
+						}
+						else
+						{
+							AddPlugin(PluginReference, ProjectReferenceChain, ExcludeFolders, NameToInstance, NameToInfo);
+						}
 					}
 				}
 			}
-			}
 
 			// Also synthesize references for plugins which are enabled by default
 			if (Rules.bCompileAgainstEngine || Rules.bCompileWithPluginSupport)
@@ -3277,15 +3277,15 @@ namespace UnrealBuildTool
 				{
 					if(Plugin.EnabledByDefault && !ReferencedNames.Contains(Plugin.Name))
 					{
-							ReferencedNames.Add(Plugin.Name);
+						ReferencedNames.Add(Plugin.Name);
 
-							PluginReferenceDescriptor PluginReference = new PluginReferenceDescriptor(Plugin.Name, null, true);
-							PluginReference.bOptional = true;
+						PluginReferenceDescriptor PluginReference = new PluginReferenceDescriptor(Plugin.Name, null, true);
+						PluginReference.bOptional = true;
 
-							AddPlugin(PluginReference, "default plugins", ExcludeFolders, NameToInstance, NameToInfo);
-						}
+						AddPlugin(PluginReference, "default plugins", ExcludeFolders, NameToInstance, NameToInfo);
 					}
 				}
+			}
 
 			// If this is a program, synthesize references for plugins which are enabled via the config file
 			if(TargetType == TargetType.Program)
@@ -3313,7 +3313,7 @@ namespace UnrealBuildTool
 			BuildPlugins = new List<UEBuildPlugin>(NameToInstance.Values);
 
 			// Determine if the project has a script plugin. We will always build UHT if there is a script plugin in the game folder.
-			bHasProjectScriptPlugin = EnabledPlugins.Any(x => x.Descriptor.bCanBeUsedWithUnrealHeaderTool && !x.File.IsUnderDirectory(UnrealBuildTool.EngineDirectory));
+			bHasProjectScriptPlugin = EnabledPlugins.Any(x => x.Descriptor.SupportedPrograms != null && x.Descriptor.SupportedPrograms.Contains("UnrealHeaderTool") && !x.File.IsUnderDirectory(UnrealBuildTool.EngineDirectory));
 		}
 
 		/// <summary>
diff --git a/Engine/Source/Programs/UnrealBuildTool/System/ExternalExecution.cs b/Engine/Source/Programs/UnrealBuildTool/System/ExternalExecution.cs
--- a/Engine/Source/Programs/UnrealBuildTool/System/ExternalExecution.cs
+++ b/Engine/Source/Programs/UnrealBuildTool/System/ExternalExecution.cs
@@ -1169,12 +1169,9 @@ namespace UnrealBuildTool
 						UBTArguments.Append(" -ignorejunk");
 
 						// Add UHT plugins to UBT command line as external plugins
-						foreach(UEBuildPlugin EnabledPlugin in Target.EnabledPlugins)
+						if(Target.ProjectFile != null)
 						{
-							if (EnabledPlugin.Descriptor.bCanBeUsedWithUnrealHeaderTool)
-							{
-								UBTArguments.Append(" -PLUGIN=\"" + EnabledPlugin.Info.File + "\"");
-							}
+							UBTArguments.AppendFormat(" -project=\"{0}\"", Target.ProjectFile);
 						}
 
 						// Add any global override for the compiler
diff --git a/Engine/Source/Programs/UnrealBuildTool/System/ModuleDescriptor.cs b/Engine/Source/Programs/UnrealBuildTool/System/ModuleDescriptor.cs
--- a/Engine/Source/Programs/UnrealBuildTool/System/ModuleDescriptor.cs
+++ b/Engine/Source/Programs/UnrealBuildTool/System/ModuleDescriptor.cs
@@ -172,6 +172,11 @@ namespace UnrealBuildTool
 		/// </summary>
 		public string[] WhitelistPrograms;
 
+		/// <summary>
+		/// List of disallowed programs
+		/// </summary>
+		public string[] BlacklistPrograms;
+
 		/// <summary>
 		/// List of additional dependencies for building this module.
 		/// </summary>
@@ -245,6 +250,12 @@ namespace UnrealBuildTool
 				Module.WhitelistPrograms = WhitelistPrograms;
 			}
 
+			string[] BlacklistPrograms;
+			if (InObject.TryGetStringArrayField("BlacklistPrograms", out BlacklistPrograms))
+			{
+				Module.BlacklistPrograms = BlacklistPrograms;
+			}
+
 			string[] AdditionalDependencies;
 			if (InObject.TryGetStringArrayField("AdditionalDependencies", out AdditionalDependencies))
 			{
@@ -322,6 +333,10 @@ namespace UnrealBuildTool
 			{
 				Writer.WriteStringArrayField("WhitelistPrograms", WhitelistPrograms);
 			}
+			if(BlacklistPrograms != null && BlacklistPrograms.Length > 0)
+			{
+				Writer.WriteStringArrayField("BlacklistPrograms", BlacklistPrograms);
+			}
 			if (AdditionalDependencies != null && AdditionalDependencies.Length > 0)
 			{
 				Writer.WriteArrayStart("AdditionalDependencies");
@@ -401,10 +416,20 @@ namespace UnrealBuildTool
 				return false;
 			}
 
-			// Check the program name is whitelisted
-			if(TargetType == TargetType.Program && (WhitelistPrograms == null || !WhitelistPrograms.Contains(TargetName)))
+			// Special checks just for programs
+			if(TargetType == TargetType.Program)
 			{
-				return false;
+				// Check the program name is whitelisted. Note that this behavior is slightly different to other whitelist/blacklist checks; we will whitelist a module of any type if it's explicitly allowed for this program.
+				if(WhitelistPrograms != null && WhitelistPrograms.Length > 0)
+				{
+					return WhitelistPrograms.Contains(TargetName);
+				}
+				
+				// Check the program name is not blacklisted
+				if(BlacklistPrograms != null && BlacklistPrograms.Contains(TargetName))
+				{
+					return false;
+				}
 			}
 
 			// Check the module is compatible with this target.
@@ -412,7 +437,7 @@ namespace UnrealBuildTool
 			{
 				case ModuleHostType.Runtime:
 				case ModuleHostType.RuntimeNoCommandlet:
-                    return true;
+                    return TargetType != TargetType.Program;
                 case ModuleHostType.CookedOnly:
                     return bBuildRequiresCookedData;
                 case ModuleHostType.RuntimeAndProgram:
diff --git a/Engine/Source/Programs/UnrealBuildTool/System/PluginDescriptor.cs b/Engine/Source/Programs/UnrealBuildTool/System/PluginDescriptor.cs
--- a/Engine/Source/Programs/UnrealBuildTool/System/PluginDescriptor.cs
+++ b/Engine/Source/Programs/UnrealBuildTool/System/PluginDescriptor.cs
@@ -149,11 +149,6 @@ namespace UnrealBuildTool
 		/// </summary>
 		public bool bIsBetaVersion;
 
-		/// <summary>
-		/// Whether this plugin can be used by UnrealHeaderTool
-		/// </summary>
-		public bool bCanBeUsedWithUnrealHeaderTool;
-
 		/// <summary>
 		/// Set for plugins which are installed
 		/// </summary>
@@ -258,7 +253,14 @@ namespace UnrealBuildTool
 			RawObject.TryGetBoolField("CanContainContent", out bCanContainContent);
 			RawObject.TryGetBoolField("IsBetaVersion", out bIsBetaVersion);
 			RawObject.TryGetBoolField("Installed", out bInstalled);
-			RawObject.TryGetBoolField("CanBeUsedWithUnrealHeaderTool", out bCanBeUsedWithUnrealHeaderTool);
+
+			bool bCanBeUsedWithUnrealHeaderTool;
+			if(RawObject.TryGetBoolField("CanBeUsedWithUnrealHeaderTool", out bCanBeUsedWithUnrealHeaderTool) && bCanBeUsedWithUnrealHeaderTool)
+			{
+				Array.Resize(ref SupportedPrograms, (SupportedPrograms == null)? 1 : SupportedPrograms.Length + 1);
+				SupportedPrograms[SupportedPrograms.Length - 1] = "UnrealHeaderTool";
+			}
+
 			RawObject.TryGetBoolField("RequiresBuildPlatform", out bRequiresBuildPlatform);
 
 			CustomBuildSteps.TryRead(RawObject, "PreBuildSteps", out PreBuildSteps);
diff --git a/Engine/Source/Runtime/Projects/Private/ModuleDescriptor.cpp b/Engine/Source/Runtime/Projects/Private/ModuleDescriptor.cpp
--- a/Engine/Source/Runtime/Projects/Private/ModuleDescriptor.cpp
+++ b/Engine/Source/Runtime/Projects/Private/ModuleDescriptor.cpp
@@ -80,6 +80,9 @@ const TCHAR* EHostType::ToString( const EHostType::Type Value )
 		case RuntimeNoCommandlet:
 			return TEXT( "RuntimeNoCommandlet" );
 
+		case RuntimeAndProgram:
+			return TEXT("RuntimeAndProgram");
+
 		case CookedOnly:
 			return TEXT("CookedOnly");
 
@@ -217,7 +220,7 @@ bool FModuleDescriptor::Read(const FJsonObject& Object, FText& OutFailReason)
 		}
 	}
 
-	// Read the supported programs
+	// Read the whitelisted programs
 	TSharedPtr<FJsonValue> WhitelistProgramsValue = Object.TryGetField(TEXT("WhitelistPrograms"));
 	if (WhitelistProgramsValue.IsValid() && WhitelistProgramsValue->Type == EJson::Array)
 	{
@@ -228,6 +231,17 @@ bool FModuleDescriptor::Read(const FJsonObject& Object, FText& OutFailReason)
 		}
 	}
 
+	// Read the blacklisted programs
+	TSharedPtr<FJsonValue> BlacklistProgramsValue = Object.TryGetField(TEXT("BlacklistPrograms"));
+	if (BlacklistProgramsValue.IsValid() && BlacklistProgramsValue->Type == EJson::Array)
+	{
+		const TArray< TSharedPtr< FJsonValue > >& ProgramsArray = BlacklistProgramsValue->AsArray();
+		for (int Idx = 0; Idx < ProgramsArray.Num(); Idx++)
+		{
+			BlacklistPrograms.Add(ProgramsArray[Idx]->AsString());
+		}
+	}
+
 	// Read the additional dependencies
 	TSharedPtr<FJsonValue> AdditionalDependenciesValue = Object.TryGetField(TEXT("AdditionalDependencies"));
 	if (AdditionalDependenciesValue.IsValid() && AdditionalDependenciesValue->Type == EJson::Array)
@@ -345,6 +359,15 @@ void FModuleDescriptor::Write(TJsonWriter<>& Writer) const
 		}
 		Writer.WriteArrayEnd();
 	}
+	if (BlacklistPrograms.Num() > 0)
+	{
+		Writer.WriteArrayStart(TEXT("BlacklistPrograms"));
+		for (int Idx = 0; Idx < BlacklistPrograms.Num(); Idx++)
+		{
+			Writer.WriteValue(BlacklistPrograms[Idx]);
+		}
+		Writer.WriteArrayEnd();
+	}
 	if (AdditionalDependencies.Num() > 0)
 	{
 		Writer.WriteArrayStart(TEXT("AdditionalDependencies"));
@@ -414,8 +437,14 @@ bool FModuleDescriptor::IsCompiledInCurrentConfiguration() const
 	}
 
 #if IS_PROGRAM
+	// Check the program is whitelisted. Note that the behavior is slightly different to other whitelist/blacklist pairs here; we will whitelist a module of any type if it's explicitly allowed for this program.
+	if (WhitelistPrograms.Num() > 0)
+	{
+		return WhitelistPrograms.Contains(UE_APP_NAME);
+	}
+
 	// Check the program is allowed
-	if (!WhitelistPrograms.Contains(UE_APP_NAME))
+	if (BlacklistPrograms.Num() > 0 && BlacklistPrograms.Contains(UE_APP_NAME))
 	{
 		return false;
 	}
@@ -426,6 +455,12 @@ bool FModuleDescriptor::IsCompiledInCurrentConfiguration() const
 	{
 	case EHostType::Runtime:
 	case EHostType::RuntimeNoCommandlet:
+		#if !IS_PROGRAM
+			return true;
+		#endif
+		break;
+
+	case EHostType::RuntimeAndProgram:
 		return true;
 
 	case EHostType::CookedOnly:
@@ -480,14 +515,20 @@ bool FModuleDescriptor::IsLoadedInCurrentConfiguration() const
 	// Check that the runtime environment allows it to be loaded
 	switch (Type)
 	{
+	case EHostType::RuntimeAndProgram:
+		#if (WITH_ENGINE || WITH_PLUGIN_SUPPORT)
+			return true;
+		#endif
+		break;
+
 	case EHostType::Runtime:
-		#if WITH_ENGINE || WITH_PLUGIN_SUPPORT
+		#if (WITH_ENGINE || WITH_PLUGIN_SUPPORT) && !IS_PROGRAM
 			return true;
 		#endif
 		break;
 	
 	case EHostType::RuntimeNoCommandlet:
-		#if WITH_ENGINE || WITH_PLUGIN_SUPPORT
+		#if (WITH_ENGINE || WITH_PLUGIN_SUPPORT)  && !IS_PROGRAM
 			if(!IsRunningCommandlet()) return true;
 		#endif
 		break;
diff --git a/Engine/Source/Runtime/Projects/Private/PluginDescriptor.cpp b/Engine/Source/Runtime/Projects/Private/PluginDescriptor.cpp
--- a/Engine/Source/Runtime/Projects/Private/PluginDescriptor.cpp
+++ b/Engine/Source/Runtime/Projects/Private/PluginDescriptor.cpp
@@ -139,6 +139,12 @@ bool FPluginDescriptor::Read(const FJsonObject& Object, FText& OutFailReason)
 	Object.TryGetBoolField(TEXT("RequiresBuildPlatform"), bRequiresBuildPlatform);
 	Object.TryGetBoolField(TEXT("Hidden"), bIsHidden);
 
+	bool bCanBeUsedWithUnrealHeaderTool;
+	if(Object.TryGetBoolField("CanBeUsedWithUnrealHeaderTool", bCanBeUsedWithUnrealHeaderTool) && bCanBeUsedWithUnrealHeaderTool)
+	{
+		SupportedPrograms.Add(TEXT("UnrealHeaderTool"));
+	}
+
 	PreBuildSteps.Read(Object, TEXT("PreBuildSteps"));
 	PostBuildSteps.Read(Object, TEXT("PostBuildSteps"));
 
diff --git a/Engine/Source/Runtime/Projects/Public/ModuleDescriptor.h b/Engine/Source/Runtime/Projects/Public/ModuleDescriptor.h
--- a/Engine/Source/Runtime/Projects/Public/ModuleDescriptor.h
+++ b/Engine/Source/Runtime/Projects/Public/ModuleDescriptor.h
@@ -69,6 +69,7 @@ namespace EHostType
 	{
 		Runtime,
 		RuntimeNoCommandlet,
+		RuntimeAndProgram,
 		CookedOnly,
 		Developer,
 		Editor,
@@ -130,9 +131,12 @@ struct PROJECTS_API FModuleDescriptor
 	/** List of disallowed target configurations */
 	TArray<FString> BlacklistTargetConfigurations;
 
-	/** List of supported programs. No blacklist, because weprograms don't support plugins by default. */
+	/** List of allowed programs */
 	TArray<FString> WhitelistPrograms;
 
+	/** List of disallowed programs */
+	TArray<FString> BlacklistPrograms;
+
 	/** List of additional dependencies for building this module. */
 	TArray<FString> AdditionalDependencies;
 
