Fix mismatch in logic between UAT's RequireTempTarget() and the engine's ProjectRequiresBuild() function. Projects need to be compiled if they disable a default plugin, as well as if they enable a plugin.

diff --git a/Engine/Source/Developer/IOS/IOSTargetPlatform/Private/IOSTargetPlatform.cpp b/Engine/Source/Developer/IOS/IOSTargetPlatform/Private/IOSTargetPlatform.cpp
--- a/Engine/Source/Developer/IOS/IOSTargetPlatform/Private/IOSTargetPlatform.cpp
+++ b/Engine/Source/Developer/IOS/IOSTargetPlatform/Private/IOSTargetPlatform.cpp
@@ -229,7 +229,7 @@ int32 FIOSTargetPlatform::CheckRequirements(const FString& ProjectPath, bool bPr
 			OutTutorialPath = FString("/Engine/Tutorial/Mobile/iOSonPCRestrictions.iOSonPCRestrictions");
 			bReadyToBuild |= ETargetPlatformReadyStatus::CodeUnsupported;
 		}
-		if (IProjectManager::Get().IsNonDefaultPluginEnabled())
+		if (!IProjectManager::Get().HasDefaultPluginSettings())
 		{
 			OutTutorialPath = FString("/Engine/Tutorial/Mobile/iOSonPCValidPlugins.iOSonPCValidPlugins");
 			bReadyToBuild |= ETargetPlatformReadyStatus::PluginsUnsupported;
diff --git a/Engine/Source/Developer/Linux/LinuxTargetPlatform/Private/LinuxTargetPlatform.h b/Engine/Source/Developer/Linux/LinuxTargetPlatform/Private/LinuxTargetPlatform.h
--- a/Engine/Source/Developer/Linux/LinuxTargetPlatform/Private/LinuxTargetPlatform.h
+++ b/Engine/Source/Developer/Linux/LinuxTargetPlatform/Private/LinuxTargetPlatform.h
@@ -219,7 +219,7 @@ public:
 				ReadyToBuild |= ETargetPlatformReadyStatus::CodeUnsupported;
 			}
 
-			if (IProjectManager::Get().IsNonDefaultPluginEnabled())
+			if (!IProjectManager::Get().HasDefaultPluginSettings())
 			{
 				ReadyToBuild |= ETargetPlatformReadyStatus::PluginsUnsupported;
 			}
diff --git a/Engine/Source/Editor/GameProjectGeneration/Private/GameProjectUtils.cpp b/Engine/Source/Editor/GameProjectGeneration/Private/GameProjectUtils.cpp
--- a/Engine/Source/Editor/GameProjectGeneration/Private/GameProjectUtils.cpp
+++ b/Engine/Source/Editor/GameProjectGeneration/Private/GameProjectUtils.cpp
@@ -3499,7 +3499,7 @@ bool GameProjectUtils::ProjectRequiresBuild(const FName InPlatformInfoName)
 	}
 
 	// check to see if any plugins beyond the defaults have been enabled
-	bRequiresBuild |= IProjectManager::Get().IsNonDefaultPluginEnabled();
+	bRequiresBuild |= !IProjectManager::Get().HasDefaultPluginSettings();
 
 	// check to see if Blueprint nativization is enabled in the Project settings
 	bRequiresBuild |= GetDefault<UProjectPackagingSettings>()->BlueprintNativizationMethod != EProjectPackagingBlueprintNativizationMethod::Disabled;
diff --git a/Engine/Source/Runtime/Projects/Private/ProjectManager.cpp b/Engine/Source/Runtime/Projects/Private/ProjectManager.cpp
--- a/Engine/Source/Runtime/Projects/Private/ProjectManager.cpp
+++ b/Engine/Source/Runtime/Projects/Private/ProjectManager.cpp
@@ -249,7 +249,7 @@ void FProjectManager::ClearSupportedTargetPlatformsForCurrentProject()
 	OnTargetPlatformsForCurrentProjectChangedEvent.Broadcast();
 }
 
-bool FProjectManager::IsNonDefaultPluginEnabled() const
+bool FProjectManager::HasDefaultPluginSettings() const
 {
 	// Get settings for the plugins which are currently enabled or disabled by the project file
 	TMap<FString, bool> ConfiguredPlugins;
@@ -274,18 +274,15 @@ bool FProjectManager::IsNonDefaultPluginEnabled() const
 				bEnabled = *EnabledPtr;
 			}
 
-			if (bEnabled)
+			bool bEnabledInDefaultExe = (Plugin->GetLoadedFrom() == EPluginLoadedFrom::Engine && Plugin->IsEnabledByDefault() && !Plugin->GetDescriptor().bInstalled);
+			if(bEnabled != bEnabledInDefaultExe)
 			{
-				bool bEnabledInDefaultExe = (Plugin->GetLoadedFrom() == EPluginLoadedFrom::Engine && Plugin->IsEnabledByDefault() && !Plugin->GetDescriptor().bInstalled);
-				if(bEnabled != bEnabledInDefaultExe)
-				{
-					return true;
-				}
+				return false;
 			}
 		}
 	}
 
-	return false;
+	return true;
 }
 
 bool FProjectManager::SetPluginEnabled(const FString& PluginName, bool bEnabled, FText& OutFailReason)
diff --git a/Engine/Source/Runtime/Projects/Private/ProjectManager.h b/Engine/Source/Runtime/Projects/Private/ProjectManager.h
--- a/Engine/Source/Runtime/Projects/Private/ProjectManager.h
+++ b/Engine/Source/Runtime/Projects/Private/ProjectManager.h
@@ -31,7 +31,7 @@ public:
 	virtual void ClearSupportedTargetPlatformsForProject(const FString& FilePath) override;
 	virtual void ClearSupportedTargetPlatformsForCurrentProject() override;
 	virtual FOnTargetPlatformsForCurrentProjectChangedEvent& OnTargetPlatformsForCurrentProjectChanged() override { return OnTargetPlatformsForCurrentProjectChangedEvent; }
-	virtual bool IsNonDefaultPluginEnabled() const override;
+	virtual bool HasDefaultPluginSettings() const override;
 	virtual bool SetPluginEnabled(const FString& PluginName, bool bEnabled, FText& OutFailReason) override;
 	virtual bool RemovePluginReference(const FString& PluginName, FText& OutFailReason) override;
 	virtual void UpdateAdditionalPluginDirectory(const FString& Dir, const bool bAddOrRemove) override;
diff --git a/Engine/Source/Runtime/Projects/Public/Interfaces/IProjectManager.h b/Engine/Source/Runtime/Projects/Public/Interfaces/IProjectManager.h
--- a/Engine/Source/Runtime/Projects/Public/Interfaces/IProjectManager.h
+++ b/Engine/Source/Runtime/Projects/Public/Interfaces/IProjectManager.h
@@ -183,11 +183,11 @@ public:
 	virtual FOnTargetPlatformsForCurrentProjectChangedEvent& OnTargetPlatformsForCurrentProjectChanged() = 0;
 
 	/**
-	 * Hack to checks whether the current project has a non-default plugin enabled (ie. one which is not included by default in UE4Game).
+	 * Checks whether the current project has default settings for plugins (ie. does not enable any new plugins, or disable any default plugins)
 	 * 
-	 * @return	True if the project has a non-default plugin enabled.
+	 * @return	True if the project has the default plugin settings.
 	 */
-	virtual bool IsNonDefaultPluginEnabled() const = 0;
+	virtual bool HasDefaultPluginSettings() const = 0;
 
 	/**
 	 * Sets whether a plugin is enabled, and updates the current project descriptor. Does not save to disk and may require restarting to load it.
