Fix issue where failure to load a plugin was a silent failure with unattended.

diff --git a/Engine/Source/Runtime/Projects/Private/PluginManager.cpp b/Engine/Source/Runtime/Projects/Private/PluginManager.cpp
--- a/Engine/Source/Runtime/Projects/Private/PluginManager.cpp
+++ b/Engine/Source/Runtime/Projects/Private/PluginManager.cpp
@@ -693,10 +693,10 @@ bool FPluginManager::ConfigureEnabledPlugin(const FPluginReferenceDescriptor& Fi
 					continue;
 				}
 
-				// If we're in unattended mode, don't open any windows
+				// If we're in unattended mode, don't open any windows and fatal out
 				if (FApp::IsUnattended())
 				{
-					UE_LOG(LogPluginManager, Error, TEXT("This project requires the '%s' plugin. Install it and try again, or remove it from the project's required plugin list."), *Reference.Name);
+					UE_LOG(LogPluginManager, Fatal, TEXT("This project requires the '%s' plugin. Install it and try again, or remove it from the project's required plugin list."), *Reference.Name);
 					return false;
 				}
 
