UBT: Plugins can now be enabled and disabled from the .target.cs file (for targets that do not use the shared compile environment), by compiling the list of enabled/disabled plugin names into the Projects module.

diff --git a/Engine/Source/Programs/UnrealBuildTool/Configuration/UEBuildTarget.cs b/Engine/Source/Programs/UnrealBuildTool/Configuration/UEBuildTarget.cs
--- a/Engine/Source/Programs/UnrealBuildTool/Configuration/UEBuildTarget.cs
+++ b/Engine/Source/Programs/UnrealBuildTool/Configuration/UEBuildTarget.cs
@@ -476,6 +476,15 @@ namespace UnrealBuildTool
 				ValidateSharedEnvironment(RulesAssembly, Desc.Name, RulesObject);
 			}
 
+			// Make sure that we don't explicitly enable or disable any plugins through the target rules. We can't 
+			if(RulesObject.BuildEnvironment == TargetBuildEnvironment.Shared && Desc.OnlyModules.Count == 0)
+			{
+				if(RulesObject.EnablePlugins.Count > 0 || RulesObject.DisablePlugins.Count > 0)
+				{
+					throw new BuildException("Explicitly enabling and disabling plugins for a target is only supported when using a unique build environment (eg. for monolithic game targets).");
+				}
+			}
+
 			// Setup the malloc profiler
 			if (RulesObject.bUseMallocProfiler)
 			{
diff --git a/Engine/Source/Runtime/Projects/Private/PluginManager.cpp b/Engine/Source/Runtime/Projects/Private/PluginManager.cpp
--- a/Engine/Source/Runtime/Projects/Private/PluginManager.cpp
+++ b/Engine/Source/Runtime/Projects/Private/PluginManager.cpp
@@ -454,6 +454,34 @@ bool FPluginManager::ConfigureEnabledPlugins()
 #if !IS_PROGRAM || HACK_HEADER_GENERATOR
 		if (!FParse::Param(FCommandLine::Get(), TEXT("NoEnginePlugins")))
 		{
+			// Configure the plugins that were enabled from the target file
+			TArray<FString> TargetEnabledPlugins = { UBT_TARGET_ENABLED_PLUGINS };
+			for (const FString& TargetEnabledPlugin : TargetEnabledPlugins)
+			{
+				if (!ConfiguredPluginNames.Contains(TargetEnabledPlugin))
+				{
+					if (!ConfigureEnabledPlugin(FPluginReferenceDescriptor(TargetEnabledPlugin, true), EnabledPluginNames))
+					{
+						return false;
+					}
+					ConfiguredPluginNames.Add(TargetEnabledPlugin);
+				}
+			}
+
+			// Configure the plugins that were enabled from the target file
+			TArray<FString> TargetDisabledPlugins = { UBT_TARGET_DISABLED_PLUGINS };
+			for (const FString& TargetDisabledPlugin : TargetDisabledPlugins)
+			{
+				if (!ConfiguredPluginNames.Contains(TargetDisabledPlugin))
+				{
+					if (!ConfigureEnabledPlugin(FPluginReferenceDescriptor(TargetDisabledPlugin, false), EnabledPluginNames))
+					{
+						return false;
+					}
+					ConfiguredPluginNames.Add(TargetDisabledPlugin);
+				}
+			}
+
 			// Find all the plugin references in the project file
 			const FProjectDescriptor* ProjectDescriptor = IProjectManager::Get().GetCurrentProject();
 			if (ProjectDescriptor != nullptr)
diff --git a/Engine/Source/Runtime/Projects/Projects.Build.cs b/Engine/Source/Runtime/Projects/Projects.Build.cs
--- a/Engine/Source/Runtime/Projects/Projects.Build.cs
+++ b/Engine/Source/Runtime/Projects/Projects.Build.cs
@@ -21,6 +21,9 @@ namespace UnrealBuildTool.Rules
 				}
 			);
 
+			PrivateDefinitions.Add(string.Format("UBT_TARGET_ENABLED_PLUGINS=TEXT(\"{0}\")", string.Join("\", TEXT(\"", Target.EnablePlugins)));
+			PrivateDefinitions.Add(string.Format("UBT_TARGET_DISABLED_PLUGINS=TEXT(\"{0}\")", string.Join("\", TEXT(\"", Target.DisablePlugins)));
+
 			if (Target.bIncludePluginsForTargetPlatforms)
 			{
 				PublicDefinitions.Add("LOAD_PLUGINS_FOR_TARGET_PLATFORMS=1");
