UBT: Fix error running programs with no explicitly enabled or disabled plugins.

#rnx

diff --git a/Engine/Source/Runtime/Projects/Projects.Build.cs b/Engine/Source/Runtime/Projects/Projects.Build.cs
--- a/Engine/Source/Runtime/Projects/Projects.Build.cs
+++ b/Engine/Source/Runtime/Projects/Projects.Build.cs
@@ -1,5 +1,8 @@
 // Copyright 1998-2018 Epic Games, Inc. All Rights Reserved.
 
+using System;
+using System.Collections.Generic;
+
 namespace UnrealBuildTool.Rules
 {
 	public class Projects : ModuleRules
@@ -21,8 +24,19 @@ namespace UnrealBuildTool.Rules
 				}
 			);
 
-			PrivateDefinitions.Add(string.Format("UBT_TARGET_ENABLED_PLUGINS=TEXT(\"{0}\")", string.Join("\", TEXT(\"", Target.EnablePlugins)));
-			PrivateDefinitions.Add(string.Format("UBT_TARGET_DISABLED_PLUGINS=TEXT(\"{0}\")", string.Join("\", TEXT(\"", Target.DisablePlugins)));
+			List<string> EnabledPluginStrings = new List<string>();
+			foreach(string EnablePlugin in Target.EnablePlugins)
+			{
+				EnabledPluginStrings.Add(String.Format("TEXT(\"{0}\")", EnablePlugin));
+			}
+			PrivateDefinitions.Add(String.Format("UBT_TARGET_ENABLED_PLUGINS={0}", String.Join(", ", EnabledPluginStrings)));
+
+			List<string> DisabledPluginStrings = new List<string>();
+			foreach(string DisablePlugin in Target.DisablePlugins)
+			{
+				DisabledPluginStrings.Add(String.Format("TEXT(\"{0}\")", DisablePlugin));
+			}
+			PrivateDefinitions.Add(String.Format("UBT_TARGET_DISABLED_PLUGINS={0}", String.Join(", ", DisabledPluginStrings)));
 
 			if (Target.bIncludePluginsForTargetPlatforms)
 			{
