Remove deprecation warning from Developer module type in the runtime. A warning for this is already output by UBT.

#rnx

diff --git a/Engine/Source/Runtime/Projects/Public/ModuleDescriptor.h b/Engine/Source/Runtime/Projects/Public/ModuleDescriptor.h
--- a/Engine/Source/Runtime/Projects/Public/ModuleDescriptor.h
+++ b/Engine/Source/Runtime/Projects/Public/ModuleDescriptor.h
@@ -88,7 +88,7 @@ namespace EHostType
 		// Deprecated due to ambiguities. Only loads in editor and program targets, but loads in any editor mode (eg. -game, -server).
 		// Use UncookedOnly for the same behavior (eg. for editor blueprint nodes needed in uncooked games), or DeveloperTool for modules
 		// that can also be loaded in cooked games but should not be shipped (eg. debugging utilities).
-		Developer UE_DEPRECATED(4.24, "The Developer module type is deprecated. Use UncookedOnly for the same behavior, or DeveloperTool for modules that can also be loaded at runtime"),
+		Developer,
 
 		// Loads on any targets where bBuildDeveloperTools is enabled.
 		DeveloperTool,
