Store statically typed configuration and target type entries in plugin and module descriptors.

#rnx

diff --git a/Engine/Source/Developer/BlueprintNativeCodeGen/Private/BlueprintNativeCodeGenUtils.cpp b/Engine/Source/Developer/BlueprintNativeCodeGen/Private/BlueprintNativeCodeGenUtils.cpp
--- a/Engine/Source/Developer/BlueprintNativeCodeGen/Private/BlueprintNativeCodeGenUtils.cpp
+++ b/Engine/Source/Developer/BlueprintNativeCodeGen/Private/BlueprintNativeCodeGenUtils.cpp
@@ -137,22 +137,22 @@ static bool BlueprintNativeCodeGenUtilsImpl::GeneratePluginDescFile(const FBluep
 				switch (PlatformInfo.PlatformType)
 				{
 				case EBuildTargetType::Game:
-					ModuleDesc->WhitelistTargets.AddUnique(TEXT("Game"));
+					ModuleDesc->WhitelistTargets.AddUnique(EBuildTargetType::Game);
 
 					// Hack to allow clients for PS4/XboxOne (etc.) to build the nativized assets plugin
 					if(!bIsClientValidForPlatform)
 					{
 						// Also add "Client" target
-						ModuleDesc->WhitelistTargets.AddUnique(TEXT("Client"));
+						ModuleDesc->WhitelistTargets.AddUnique(EBuildTargetType::Client);
 					}
 					break;
 
 				case EBuildTargetType::Client:
-					ModuleDesc->WhitelistTargets.AddUnique(TEXT("Client"));
+					ModuleDesc->WhitelistTargets.AddUnique(EBuildTargetType::Client);
 					break;
 
 				case EBuildTargetType::Server:
-					ModuleDesc->WhitelistTargets.AddUnique(TEXT("Server"));
+					ModuleDesc->WhitelistTargets.AddUnique(EBuildTargetType::Server);
 					break;
 
 				case EBuildTargetType::Editor:
diff --git a/Engine/Source/Runtime/Json/Public/Dom/JsonObject.h b/Engine/Source/Runtime/Json/Public/Dom/JsonObject.h
--- a/Engine/Source/Runtime/Json/Public/Dom/JsonObject.h
+++ b/Engine/Source/Runtime/Json/Public/Dom/JsonObject.h
@@ -145,6 +145,28 @@ public:
 	/** Get the field named FieldName as an array of strings. Returns false if it doesn't exist or any member cannot be converted. */
 	bool TryGetStringArrayField(const FString& FieldName, TArray<FString>& OutArray) const;
 
+	/** Get the field named FieldName as an array of enums. Returns false if it doesn't exist or any member is not a string. */
+	template<typename TEnum>
+	bool TryGetEnumArrayField(const FString& FieldName, TArray<TEnum>& OutArray) const
+	{
+		TArray<FString> Strings;
+		if (!TryGetStringArrayField(FieldName, Strings))
+		{
+			return false;
+		}
+
+		OutArray.Empty();
+		for (const FString& String : Strings)
+		{
+			TEnum Value;
+			if (LexTryParseString(Value, *String))
+			{
+				OutArray.Add(Value);
+			}
+		}
+		return true;
+	}
+
 	/** Add a field named FieldName with value of StringValue */
 	void SetStringField( const FString& FieldName, const FString& StringValue );
 
diff --git a/Engine/Source/Runtime/Projects/Private/ModuleDescriptor.cpp b/Engine/Source/Runtime/Projects/Private/ModuleDescriptor.cpp
--- a/Engine/Source/Runtime/Projects/Private/ModuleDescriptor.cpp
+++ b/Engine/Source/Runtime/Projects/Private/ModuleDescriptor.cpp
@@ -169,104 +169,24 @@ bool FModuleDescriptor::Read(const FJsonObject& Object, FText& OutFailReason)
 		}
 	}
 
-	// Read the whitelisted platforms
-	TSharedPtr<FJsonValue> WhitelistPlatformsValue = Object.TryGetField(TEXT("WhitelistPlatforms"));
-	if(WhitelistPlatformsValue.IsValid() && WhitelistPlatformsValue->Type == EJson::Array)
-	{
-		const TArray< TSharedPtr< FJsonValue > >& PlatformsArray = WhitelistPlatformsValue->AsArray();
-		for(int Idx = 0; Idx < PlatformsArray.Num(); Idx++)
-		{
-			WhitelistPlatforms.Add(PlatformsArray[Idx]->AsString());
-		}
-	}
-
-	// Read the blacklisted platforms
-	TSharedPtr<FJsonValue> BlacklistPlatformsValue = Object.TryGetField(TEXT("BlacklistPlatforms"));
-	if(BlacklistPlatformsValue.IsValid() && BlacklistPlatformsValue->Type == EJson::Array)
-	{
-		const TArray< TSharedPtr< FJsonValue > >& PlatformsArray = BlacklistPlatformsValue->AsArray();
-		for(int Idx = 0; Idx < PlatformsArray.Num(); Idx++)
-		{
-			BlacklistPlatforms.Add(PlatformsArray[Idx]->AsString());
-		}
-	}
-
-	// Read the whitelisted targets
-	TSharedPtr<FJsonValue> WhitelistTargetsValue = Object.TryGetField(TEXT("WhitelistTargets"));
-	if (WhitelistTargetsValue.IsValid() && WhitelistTargetsValue->Type == EJson::Array)
-	{
-		const TArray< TSharedPtr< FJsonValue > >& TargetsArray = WhitelistTargetsValue->AsArray();
-		for (int Idx = 0; Idx < TargetsArray.Num(); Idx++)
-		{
-			WhitelistTargets.Add(TargetsArray[Idx]->AsString());
-		}
-	}
+	// Read the whitelisted and blacklisted platforms
+	Object.TryGetStringArrayField(TEXT("WhitelistPlatforms"), WhitelistPlatforms);
+	Object.TryGetStringArrayField(TEXT("BlacklistPlatforms"), BlacklistPlatforms);
 
-	// Read the blacklisted targets
-	TSharedPtr<FJsonValue> BlacklistTargetsValue = Object.TryGetField(TEXT("BlacklistTargets"));
-	if (BlacklistTargetsValue.IsValid() && BlacklistTargetsValue->Type == EJson::Array)
-	{
-		const TArray< TSharedPtr< FJsonValue > >& TargetsArray = BlacklistTargetsValue->AsArray();
-		for (int Idx = 0; Idx < TargetsArray.Num(); Idx++)
-		{
-			BlacklistTargets.Add(TargetsArray[Idx]->AsString());
-		}
-	}
+	// Read the whitelisted and blacklisted targets
+	Object.TryGetEnumArrayField(TEXT("WhitelistTargets"), WhitelistTargets);
+	Object.TryGetEnumArrayField(TEXT("BlacklistTargets"), BlacklistTargets);
 
-	// Read the whitelisted target configurations
-	TSharedPtr<FJsonValue> WhitelistTargetConfigurationsValue = Object.TryGetField(TEXT("WhitelistTargetConfigurations"));
-	if (WhitelistTargetConfigurationsValue.IsValid() && WhitelistTargetConfigurationsValue->Type == EJson::Array)
-	{
-		const TArray< TSharedPtr< FJsonValue > >& ConfigsArray = WhitelistTargetConfigurationsValue->AsArray();
-		for (int Idx = 0; Idx < ConfigsArray.Num(); Idx++)
-		{
-			WhitelistTargetConfigurations.Add(ConfigsArray[Idx]->AsString());
-		}
-	}
+	// Read the whitelisted and blacklisted target configurations
+	Object.TryGetEnumArrayField(TEXT("WhitelistTargetConfigurations"), WhitelistTargetConfigurations);
+	Object.TryGetEnumArrayField(TEXT("BlacklistTargetConfigurations"), BlacklistTargetConfigurations);
 
-	// Read the blacklisted target configurations
-	TSharedPtr<FJsonValue> BlacklistTargetConfigurationsValue = Object.TryGetField(TEXT("BlacklistTargetConfigurations"));
-	if (BlacklistTargetConfigurationsValue.IsValid() && BlacklistTargetConfigurationsValue->Type == EJson::Array)
-	{
-		const TArray< TSharedPtr< FJsonValue > >& ConfigsArray = BlacklistTargetConfigurationsValue->AsArray();
-		for (int Idx = 0; Idx < ConfigsArray.Num(); Idx++)
-		{
-			BlacklistTargetConfigurations.Add(ConfigsArray[Idx]->AsString());
-		}
-	}
-
-	// Read the whitelisted programs
-	TSharedPtr<FJsonValue> WhitelistProgramsValue = Object.TryGetField(TEXT("WhitelistPrograms"));
-	if (WhitelistProgramsValue.IsValid() && WhitelistProgramsValue->Type == EJson::Array)
-	{
-		const TArray< TSharedPtr< FJsonValue > >& ProgramsArray = WhitelistProgramsValue->AsArray();
-		for (int Idx = 0; Idx < ProgramsArray.Num(); Idx++)
-		{
-			WhitelistPrograms.Add(ProgramsArray[Idx]->AsString());
-		}
-	}
-
-	// Read the blacklisted programs
-	TSharedPtr<FJsonValue> BlacklistProgramsValue = Object.TryGetField(TEXT("BlacklistPrograms"));
-	if (BlacklistProgramsValue.IsValid() && BlacklistProgramsValue->Type == EJson::Array)
-	{
-		const TArray< TSharedPtr< FJsonValue > >& ProgramsArray = BlacklistProgramsValue->AsArray();
-		for (int Idx = 0; Idx < ProgramsArray.Num(); Idx++)
-		{
-			BlacklistPrograms.Add(ProgramsArray[Idx]->AsString());
-		}
-	}
+	// Read the whitelisted and blacklisted programs
+	Object.TryGetStringArrayField(TEXT("WhitelistPrograms"), WhitelistPrograms);
+	Object.TryGetStringArrayField(TEXT("BlacklistPrograms"), BlacklistPrograms);
 
 	// Read the additional dependencies
-	TSharedPtr<FJsonValue> AdditionalDependenciesValue = Object.TryGetField(TEXT("AdditionalDependencies"));
-	if (AdditionalDependenciesValue.IsValid() && AdditionalDependenciesValue->Type == EJson::Array)
-	{
-		const TArray< TSharedPtr< FJsonValue > >& DepArray = AdditionalDependenciesValue->AsArray();
-		for (int Idx = 0; Idx < DepArray.Num(); Idx++)
-		{
-			AdditionalDependencies.Add(DepArray[Idx]->AsString());
-		}
-	}
+	Object.TryGetStringArrayField(TEXT("AdditionalDependencies"), AdditionalDependencies);
 
 	return true;
 }
@@ -334,7 +254,7 @@ void FModuleDescriptor::Write(TJsonWriter<>& Writer) const
 		Writer.WriteArrayStart(TEXT("WhitelistTargets"));
 		for (int Idx = 0; Idx < WhitelistTargets.Num(); Idx++)
 		{
-			Writer.WriteValue(WhitelistTargets[Idx]);
+			Writer.WriteValue(LexToString(WhitelistTargets[Idx]));
 		}
 		Writer.WriteArrayEnd();
 	}
@@ -343,7 +263,7 @@ void FModuleDescriptor::Write(TJsonWriter<>& Writer) const
 		Writer.WriteArrayStart(TEXT("BlacklistTargets"));
 		for (int Idx = 0; Idx < BlacklistTargets.Num(); Idx++)
 		{
-			Writer.WriteValue(BlacklistTargets[Idx]);
+			Writer.WriteValue(LexToString(BlacklistTargets[Idx]));
 		}
 		Writer.WriteArrayEnd();
 	}
@@ -352,7 +272,7 @@ void FModuleDescriptor::Write(TJsonWriter<>& Writer) const
 		Writer.WriteArrayStart(TEXT("WhitelistTargetConfigurations"));
 		for (int Idx = 0; Idx < WhitelistTargetConfigurations.Num(); Idx++)
 		{
-			Writer.WriteValue(WhitelistTargetConfigurations[Idx]);
+			Writer.WriteValue(LexToString(WhitelistTargetConfigurations[Idx]));
 		}
 		Writer.WriteArrayEnd();
 	}
@@ -361,7 +281,7 @@ void FModuleDescriptor::Write(TJsonWriter<>& Writer) const
 		Writer.WriteArrayStart(TEXT("BlacklistTargetConfigurations"));
 		for (int Idx = 0; Idx < BlacklistTargetConfigurations.Num(); Idx++)
 		{
-			Writer.WriteValue(BlacklistTargetConfigurations[Idx]);
+			Writer.WriteValue(LexToString(BlacklistTargetConfigurations[Idx]));
 		}
 		Writer.WriteArrayEnd();
 	}
@@ -408,135 +328,97 @@ void FModuleDescriptor::WriteArray(TJsonWriter<>& Writer, const TCHAR* Name, con
 	}
 }
 
-bool FModuleDescriptor::IsCompiledInCurrentConfiguration() const
+bool FModuleDescriptor::IsCompiledInConfiguration(const FString& Platform, EBuildConfiguration Configuration, const FString& TargetName, EBuildTargetType TargetType, bool bBuildDeveloperTools, bool bBuildRequiresCookedData) const
 {
-	// Cache the string for the current platform
-	static const FString UBTPlatform = FPlatformMisc::GetUBTPlatform();
-	static const FString UBTTarget = FPlatformMisc::GetUBTTarget();
-	static const FString UBTTargetConfiguration = LexToString(FApp::GetBuildConfiguration());
-
 	// Check the platform is whitelisted
-	if(WhitelistPlatforms.Num() > 0 && !WhitelistPlatforms.Contains(UBTPlatform))
+	if (WhitelistPlatforms.Num() > 0 && !WhitelistPlatforms.Contains(Platform))
 	{
 		return false;
 	}
 
 	// Check the platform is not blacklisted
-	if(BlacklistPlatforms.Num() > 0 && BlacklistPlatforms.Contains(UBTPlatform))
+	if (BlacklistPlatforms.Contains(Platform))
 	{
 		return false;
 	}
 
 	// Check the target is whitelisted
-	if (WhitelistTargets.Num() > 0 && !WhitelistTargets.Contains(UBTTarget))
+	if (WhitelistTargets.Num() > 0 && !WhitelistTargets.Contains(TargetType))
 	{
 		return false;
 	}
 
 	// Check the target is not blacklisted
-	if (BlacklistTargets.Num() > 0 && BlacklistTargets.Contains(UBTTarget))
+	if (BlacklistTargets.Contains(TargetType))
 	{
 		return false;
 	}
 
 	// Check the target configuration is whitelisted
-	if (WhitelistTargetConfigurations.Num() > 0 && !WhitelistTargetConfigurations.Contains(UBTTargetConfiguration))
+	if (WhitelistTargetConfigurations.Num() > 0 && !WhitelistTargetConfigurations.Contains(Configuration))
 	{
 		return false;
 	}
 
 	// Check the target configuration is not blacklisted
-	if (BlacklistTargetConfigurations.Num() > 0 && BlacklistTargetConfigurations.Contains(UBTTargetConfiguration))
+	if (BlacklistTargetConfigurations.Contains(Configuration))
 	{
 		return false;
 	}
 
-#if IS_PROGRAM
-	// Check the program is whitelisted. Note that the behavior is slightly different to other whitelist/blacklist pairs here; we will whitelist a module of any type if it's explicitly allowed for this program.
-	if (WhitelistPrograms.Num() > 0)
-	{
-		return WhitelistPrograms.Contains(UE_APP_NAME);
-	}
-
-	// Check the program is allowed
-	if (BlacklistPrograms.Num() > 0 && BlacklistPrograms.Contains(UE_APP_NAME))
+	// Special checks just for programs
+	if(TargetType == EBuildTargetType::Program)
 	{
-		return false;
+		// Check the program name is whitelisted. Note that this behavior is slightly different to other whitelist/blacklist checks; we will whitelist a module of any type if it's explicitly allowed for this program.
+		if(WhitelistPrograms.Num() > 0)
+		{
+			return WhitelistPrograms.Contains(TargetName);
+		}
+				
+		// Check the program name is not blacklisted
+		if(BlacklistPrograms.Contains(TargetName))
+		{
+			return false;
+		}
 	}
-#endif
 
-	// Check the module is compatible with this target. This should match ModuleDescriptor.IsCompiledInConfiguration in UBT
+	// Check the module is compatible with this target.
 	switch (Type)
 	{
 	case EHostType::Runtime:
 	case EHostType::RuntimeNoCommandlet:
-		#if !IS_PROGRAM
-			return true;
-		#endif
-		break;
-
+        return TargetType != EBuildTargetType::Program;
 	case EHostType::RuntimeAndProgram:
 		return true;
-
 	case EHostType::CookedOnly:
-		return FPlatformProperties::RequiresCookedData();
-
+        return bBuildRequiresCookedData;
 	case EHostType::UncookedOnly:
-		return !FPlatformProperties::RequiresCookedData();
-
+		return !bBuildRequiresCookedData;
 	case EHostType::Developer:
-		#if WITH_EDITOR || IS_PROGRAM
-			return true;
-		#else
-			return false;
-		#endif
-
+		return TargetType == EBuildTargetType::Editor || TargetType == EBuildTargetType::Program;
 	case EHostType::DeveloperTool:
-		#if WITH_UNREAL_DEVELOPER_TOOLS
-			return true;
-		#else
-			return false;		
-		#endif
-
+		return bBuildDeveloperTools;
 	case EHostType::Editor:
 	case EHostType::EditorNoCommandlet:
-		#if WITH_EDITOR
-			return true;
-		#endif
-		break;
-
+		return TargetType == EBuildTargetType::Editor;
 	case EHostType::EditorAndProgram:
-		#if WITH_EDITOR || IS_PROGRAM
-			return true;
-		#else
-			return false;
-		#endif
-
+		return TargetType == EBuildTargetType::Editor || TargetType == EBuildTargetType::Program;
 	case EHostType::Program:
-		#if IS_PROGRAM
-			return true;
-		#endif
-		break;
-
-	case EHostType::ServerOnly:
-#if IS_PROGRAM
-		return false;
-#else
-		return !FPlatformProperties::IsClientOnly();
-#endif
-
-	case EHostType::ClientOnly:
-#if IS_PROGRAM
-		return false;
-#else
-		return !FPlatformProperties::IsServerOnly();
-#endif
-
-	}
+		return TargetType == EBuildTargetType::Program;
+    case EHostType::ServerOnly:
+        return TargetType != EBuildTargetType::Program && TargetType != EBuildTargetType::Client;
+    case EHostType::ClientOnly:
+        return TargetType != EBuildTargetType::Program && TargetType != EBuildTargetType::Server;
+    }
 
 	return false;
 }
 
+bool FModuleDescriptor::IsCompiledInCurrentConfiguration() const
+{
+	return IsCompiledInConfiguration(FPlatformMisc::GetUBTPlatform(), FApp::GetBuildConfiguration(), UE_APP_NAME, FApp::GetBuildTargetType(), !!WITH_UNREAL_DEVELOPER_TOOLS, FPlatformProperties::RequiresCookedData());
+}
+
 bool FModuleDescriptor::IsLoadedInCurrentConfiguration() const
 {
 	// Check that the module is built for this configuration
diff --git a/Engine/Source/Runtime/Projects/Private/PluginManager.cpp b/Engine/Source/Runtime/Projects/Private/PluginManager.cpp
--- a/Engine/Source/Runtime/Projects/Private/PluginManager.cpp
+++ b/Engine/Source/Runtime/Projects/Private/PluginManager.cpp
@@ -756,7 +756,7 @@ bool FPluginManager::ConfigureEnabledPlugin(const FPluginReferenceDescriptor& Fi
 			const FPluginReferenceDescriptor& Reference = *NewPluginReferences[Idx];
 
 			// Check if the plugin is required for this platform
-			if(!Reference.IsEnabledForPlatform(FPlatformMisc::GetUBTPlatform()) || !Reference.IsEnabledForTargetConfiguration(LexToString(FApp::GetBuildConfiguration())) || !Reference.IsEnabledForTarget(FPlatformMisc::GetUBTTarget()))
+			if(!Reference.IsEnabledForPlatform(FPlatformMisc::GetUBTPlatform()) || !Reference.IsEnabledForTargetConfiguration(FApp::GetBuildConfiguration()) || !Reference.IsEnabledForTarget(FApp::GetBuildTargetType()))
 			{
 				UE_LOG(LogPluginManager, Verbose, TEXT("Ignoring plugin '%s' for platform/configuration"), *Reference.Name);
 				continue;
diff --git a/Engine/Source/Runtime/Projects/Private/PluginReferenceDescriptor.cpp b/Engine/Source/Runtime/Projects/Private/PluginReferenceDescriptor.cpp
--- a/Engine/Source/Runtime/Projects/Private/PluginReferenceDescriptor.cpp
+++ b/Engine/Source/Runtime/Projects/Private/PluginReferenceDescriptor.cpp
@@ -38,7 +38,7 @@ bool FPluginReferenceDescriptor::IsEnabledForPlatform( const FString& Platform )
 	return true;
 }
 
-bool FPluginReferenceDescriptor::IsEnabledForTarget(const FString& Target) const
+bool FPluginReferenceDescriptor::IsEnabledForTarget(EBuildTargetType TargetType) const
 {
     // If it's not enabled at all, return false
     if (!bEnabled)
@@ -47,13 +47,13 @@ bool FPluginReferenceDescriptor::IsEnabledForTarget(const FString& Target) const
     }
 
     // If there is a list of whitelisted platforms, and this isn't one of them, return false
-    if (WhitelistTargets.Num() > 0 && !WhitelistTargets.Contains(Target))
+    if (WhitelistTargets.Num() > 0 && !WhitelistTargets.Contains(TargetType))
     {
         return false;
     }
 
     // If this platform is blacklisted, also return false
-    if (BlacklistTargets.Contains(Target))
+    if (BlacklistTargets.Contains(TargetType))
     {
         return false;
     }
@@ -61,7 +61,7 @@ bool FPluginReferenceDescriptor::IsEnabledForTarget(const FString& Target) const
     return true;
 }
 
-bool FPluginReferenceDescriptor::IsEnabledForTargetConfiguration(const FString& TargetConfiguration) const
+bool FPluginReferenceDescriptor::IsEnabledForTargetConfiguration(EBuildConfiguration Configuration) const
 {
 	// If it's not enabled at all, return false
 	if (!bEnabled)
@@ -70,13 +70,13 @@ bool FPluginReferenceDescriptor::IsEnabledForTargetConfiguration(const FString&
 	}
 
 	// If there is a list of whitelisted target configurations, and this isn't one of them, return false
-	if (WhitelistTargetConfigurations.Num() > 0 && !WhitelistTargetConfigurations.Contains(TargetConfiguration))
+	if (WhitelistTargetConfigurations.Num() > 0 && !WhitelistTargetConfigurations.Contains(Configuration))
 	{
 		return false;
 	}
 
 	// If this target configuration is blacklisted, also return false
-	if (BlacklistTargetConfigurations.Contains(TargetConfiguration))
+	if (BlacklistTargetConfigurations.Contains(Configuration))
 	{
 		return false;
 	}
@@ -117,12 +117,12 @@ bool FPluginReferenceDescriptor::Read( const FJsonObject& Object, FText& OutFail
 	Object.TryGetStringArrayField(TEXT("BlacklistPlatforms"), BlacklistPlatforms);
 
 	// Get the target configuration lists
-	Object.TryGetStringArrayField(TEXT("WhitelistTargetConfigurations"), WhitelistTargetConfigurations);
-	Object.TryGetStringArrayField(TEXT("BlacklistTargetConfigurations"), BlacklistTargetConfigurations);
+	Object.TryGetEnumArrayField(TEXT("WhitelistTargetConfigurations"), WhitelistTargetConfigurations);
+	Object.TryGetEnumArrayField(TEXT("BlacklistTargetConfigurations"), BlacklistTargetConfigurations);
 
 	// Get the target lists
-	Object.TryGetStringArrayField(TEXT("WhitelistTargets"), WhitelistTargets);
-	Object.TryGetStringArrayField(TEXT("BlacklistTargets"), BlacklistTargets);
+	Object.TryGetEnumArrayField(TEXT("WhitelistTargets"), WhitelistTargets);
+	Object.TryGetEnumArrayField(TEXT("BlacklistTargets"), BlacklistTargets);
 
 	// Get the supported platform list
 	Object.TryGetStringArrayField(TEXT("SupportedTargetPlatforms"), SupportedTargetPlatforms);
@@ -210,7 +210,7 @@ void FPluginReferenceDescriptor::Write( TJsonWriter<>& Writer ) const
 
 		for (int Idx = 0; Idx < WhitelistTargetConfigurations.Num(); Idx++)
 		{
-			Writer.WriteValue(WhitelistTargetConfigurations[Idx]);
+			Writer.WriteValue(LexToString(WhitelistTargetConfigurations[Idx]));
 		}
 
 		Writer.WriteArrayEnd();
@@ -222,7 +222,7 @@ void FPluginReferenceDescriptor::Write( TJsonWriter<>& Writer ) const
 
 		for (int Idx = 0; Idx < BlacklistTargetConfigurations.Num(); Idx++)
 		{
-			Writer.WriteValue(BlacklistTargetConfigurations[Idx]);
+			Writer.WriteValue(LexToString(BlacklistTargetConfigurations[Idx]));
 		}
 
 		Writer.WriteArrayEnd();
@@ -234,7 +234,7 @@ void FPluginReferenceDescriptor::Write( TJsonWriter<>& Writer ) const
 
 		for (int Idx = 0; Idx < WhitelistTargets.Num(); Idx++)
 		{
-			Writer.WriteValue(WhitelistTargets[Idx]);
+			Writer.WriteValue(LexToString(WhitelistTargets[Idx]));
 		}
 
 		Writer.WriteArrayEnd();
@@ -246,7 +246,7 @@ void FPluginReferenceDescriptor::Write( TJsonWriter<>& Writer ) const
 
 		for (int Idx = 0; Idx < BlacklistTargets.Num(); Idx++)
 		{
-			Writer.WriteValue(BlacklistTargets[Idx]);
+			Writer.WriteValue(LexToString(BlacklistTargets[Idx]));
 		}
 
 		Writer.WriteArrayEnd();
diff --git a/Engine/Source/Runtime/Projects/Public/ModuleDescriptor.h b/Engine/Source/Runtime/Projects/Public/ModuleDescriptor.h
--- a/Engine/Source/Runtime/Projects/Public/ModuleDescriptor.h
+++ b/Engine/Source/Runtime/Projects/Public/ModuleDescriptor.h
@@ -156,16 +156,16 @@ struct PROJECTS_API FModuleDescriptor
 	TArray<FString> BlacklistPlatforms;
 
 	/** List of allowed targets */
-	TArray<FString> WhitelistTargets;
+	TArray<EBuildTargetType> WhitelistTargets;
 
 	/** List of disallowed targets */
-	TArray<FString> BlacklistTargets;
+	TArray<EBuildTargetType> BlacklistTargets;
 
 	/** List of allowed target configurations */
-	TArray<FString> WhitelistTargetConfigurations;
+	TArray<EBuildConfiguration> WhitelistTargetConfigurations;
 
 	/** List of disallowed target configurations */
-	TArray<FString> BlacklistTargetConfigurations;
+	TArray<EBuildConfiguration> BlacklistTargetConfigurations;
 
 	/** List of allowed programs */
 	TArray<FString> WhitelistPrograms;
@@ -191,6 +191,9 @@ struct PROJECTS_API FModuleDescriptor
 	/** Writes an array of modules to JSON */
 	static void WriteArray(TJsonWriter<>& Writer, const TCHAR* Name, const TArray<FModuleDescriptor>& Modules);
 
+	/** Tests whether the module should be built for the given target */
+	bool IsCompiledInConfiguration(const FString& Platform, EBuildConfiguration Configuration, const FString& TargetName, EBuildTargetType TargetType, bool bBuildDeveloperTools, bool bBuildRequiresCookedData) const;
+
 	/** Tests whether the module should be built for the current engine configuration */
 	bool IsCompiledInCurrentConfiguration() const;
 
diff --git a/Engine/Source/Runtime/Projects/Public/PluginReferenceDescriptor.h b/Engine/Source/Runtime/Projects/Public/PluginReferenceDescriptor.h
--- a/Engine/Source/Runtime/Projects/Public/PluginReferenceDescriptor.h
+++ b/Engine/Source/Runtime/Projects/Public/PluginReferenceDescriptor.h
@@ -34,16 +34,16 @@ struct PROJECTS_API FPluginReferenceDescriptor
 	TArray<FString> BlacklistPlatforms;
  
 	/** If enabled, list of target configurations for which the plugin should be enabled (or all target configurations if blank). */
-	TArray<FString> WhitelistTargetConfigurations;
+	TArray<EBuildConfiguration> WhitelistTargetConfigurations;
 
 	/** If enabled, list of target configurations for which the plugin should be disabled. */
-	TArray<FString> BlacklistTargetConfigurations;
+	TArray<EBuildConfiguration> BlacklistTargetConfigurations;
 
 	/** If enabled, list of targets for which the plugin should be enabled (or all targets if blank). */
-	TArray<FString> WhitelistTargets;
+	TArray<EBuildTargetType> WhitelistTargets;
 
 	/** If enabled, list of targets for which the plugin should be disabled. */
-	TArray<FString> BlacklistTargets;
+	TArray<EBuildTargetType> BlacklistTargets;
 
 	/** The list of supported target platforms for this plugin. This field is copied from the plugin descriptor, and supplements the user's whitelisted and blacklisted platforms. */
 	TArray<FString> SupportedTargetPlatforms;
@@ -55,10 +55,10 @@ struct PROJECTS_API FPluginReferenceDescriptor
 	bool IsEnabledForPlatform(const FString& Platform) const;
 
 	/** Determines whether the plugin is enabled for the given target configuration */
-	bool IsEnabledForTargetConfiguration(const FString& TargetConfiguration) const;
+	bool IsEnabledForTargetConfiguration(EBuildConfiguration Configuration) const;
 
 	/** Determines whether the plugin is enabled for the given target */
-	bool IsEnabledForTarget(const FString& Target) const;
+	bool IsEnabledForTarget(EBuildTargetType TargetType) const;
 
 	/** Determines if the referenced plugin is supported for the given platform */
 	bool IsSupportedTargetPlatform(const FString& Platform) const;
