Allow respect the explicit module whitelist/blacklist for enabling/disabling modules in program targets.

#rnx

diff --git a/Engine/Source/Runtime/Projects/Private/ModuleDescriptor.cpp b/Engine/Source/Runtime/Projects/Private/ModuleDescriptor.cpp
--- a/Engine/Source/Runtime/Projects/Private/ModuleDescriptor.cpp
+++ b/Engine/Source/Runtime/Projects/Private/ModuleDescriptor.cpp
@@ -427,6 +427,25 @@ bool FModuleDescriptor::IsLoadedInCurrentConfiguration() const
 		return false;
 	}
 
+	// Always respect the whitelist/blacklist for program targets
+	EBuildTargetType TargetType = FApp::GetBuildTargetType();
+	if(TargetType == EBuildTargetType::Program)
+	{
+		const FString TargetName = UE_APP_NAME;
+
+		// Check the program name is whitelisted. Note that this behavior is slightly different to other whitelist/blacklist checks; we will whitelist a module of any type if it's explicitly allowed for this program.
+		if(WhitelistPrograms.Num() > 0)
+		{
+			return WhitelistPrograms.Contains(TargetName);
+		}
+				
+		// Check the program name is not blacklisted
+		if(BlacklistPrograms.Contains(TargetName))
+		{
+			return false;
+		}
+	}
+
 	// Check that the runtime environment allows it to be loaded
 	switch (Type)
 	{
