Fixed platform extensions functionality for plugins
- Add AdditionalRestrictedFolders to FDataDrivenPlatformInfoRegistry
- Fixed FPluginManager::ReadAllPlugins() to scan AdditionalRestrictedFolders to find related platform directory for plugin
- Fixed FPluginManager::CreatePluginObject() complaining about existing plugin found twice (ignore if same source path)
- Fixed Plugins::TryMergeWithParent() handling for blacklist merge with child - if parent had blacklist but child didn't, child was incorrectly added to blacklist

#rnx

diff --git a/Engine/Source/Programs/UnrealBuildTool/System/Plugins.cs b/Engine/Source/Programs/UnrealBuildTool/System/Plugins.cs
--- a/Engine/Source/Programs/UnrealBuildTool/System/Plugins.cs
+++ b/Engine/Source/Programs/UnrealBuildTool/System/Plugins.cs
@@ -396,16 +396,12 @@ namespace UnrealBuildTool
 						}
 						if (ParentModule.BlacklistPlatforms != null && ParentModule.BlacklistPlatforms.Length > 0)
 						{
-							List<UnrealTargetPlatform> Blacklist = ParentModule.BlacklistPlatforms.ToList();
 							if (ChildModule.BlacklistPlatforms != null && ChildModule.BlacklistPlatforms.Length > 0)
 							{
+								List<UnrealTargetPlatform> Blacklist = ParentModule.BlacklistPlatforms.ToList();
 								Blacklist.AddRange(ChildModule.BlacklistPlatforms);
+								ParentModule.BlacklistPlatforms = Blacklist.ToArray();
 							}
-							else
-							{
-								Blacklist.Add(Platform);
-							}
-							ParentModule.BlacklistPlatforms = Blacklist.ToArray();
 						}
 					}
 				}
diff --git a/Engine/Source/Runtime/Core/Private/Misc/DataDrivenPlatformInfoRegistry.cpp b/Engine/Source/Runtime/Core/Private/Misc/DataDrivenPlatformInfoRegistry.cpp
--- a/Engine/Source/Runtime/Core/Private/Misc/DataDrivenPlatformInfoRegistry.cpp
+++ b/Engine/Source/Runtime/Core/Private/Misc/DataDrivenPlatformInfoRegistry.cpp
@@ -91,6 +91,7 @@ const TMap<FString, FDataDrivenPlatformInfoRegistry::FPlatformInfo>& FDataDriven
 				FDataDrivenPlatformInfoRegistry::FPlatformInfo& Info = DataDrivenPlatforms.Add(PlatformName, FDataDrivenPlatformInfoRegistry::FPlatformInfo());
 				IniFile.GetBool(TEXT("DataDrivenPlatformInfo"), TEXT("bIsConfidential"), Info.bIsConfidential);
 				IniFile.GetBool(TEXT("DataDrivenPlatformInfo"), TEXT("bRestrictLocalization"), Info.bRestrictLocalization);
+				IniFile.GetArray(TEXT("DataDrivenPlatformInfo"), TEXT("AdditionalRestrictedFolders"), Info.AdditionalRestrictedFolders);
 
 				// get the parent to build list later
 				FString IniParent;
diff --git a/Engine/Source/Runtime/Core/Public/Misc/DataDrivenPlatformInfoRegistry.h b/Engine/Source/Runtime/Core/Public/Misc/DataDrivenPlatformInfoRegistry.h
--- a/Engine/Source/Runtime/Core/Public/Misc/DataDrivenPlatformInfoRegistry.h
+++ b/Engine/Source/Runtime/Core/Public/Misc/DataDrivenPlatformInfoRegistry.h
@@ -20,6 +20,9 @@ struct CORE_API FDataDrivenPlatformInfoRegistry
 		// should this platform be split when using ELocTextPlatformSplitMode::Restricted (only used when bIsConfidential is true)
 		bool bRestrictLocalization = false;
 
+		// list of additonal restricted folders
+		TArray<FString> AdditionalRestrictedFolders;
+
 		// cached list of ini parents
 		TArray<FString> IniParentChain;
 	};
diff --git a/Engine/Source/Runtime/Projects/Private/PluginManager.cpp b/Engine/Source/Runtime/Projects/Private/PluginManager.cpp
--- a/Engine/Source/Runtime/Projects/Private/PluginManager.cpp
+++ b/Engine/Source/Runtime/Projects/Private/PluginManager.cpp
@@ -230,6 +230,10 @@ void FPluginManager::ReadAllPlugins(TMap<FString, TSharedRef<FPlugin>>& Plugins,
 		ReadPluginsInDirectory(FPaths::EnginePluginsDir(), EPluginType::Engine, Plugins, ChildPlugins);
 		for (auto const& Platform : FDataDrivenPlatformInfoRegistry::GetAllPlatformInfos())
 		{
+			for (auto const& AdditionalFolder : Platform.Value.AdditionalRestrictedFolders)
+			{
+				ReadPluginsInDirectory(FPaths::EnginePlatformExtensionsDir() / AdditionalFolder / TEXT("Plugins"), EPluginType::Engine, Plugins, ChildPlugins);
+			}
 			ReadPluginsInDirectory(FPaths::EnginePlatformExtensionsDir() / Platform.Key / TEXT("Plugins"), EPluginType::Engine, Plugins, ChildPlugins);
 		}
 
@@ -241,6 +245,10 @@ void FPluginManager::ReadAllPlugins(TMap<FString, TSharedRef<FPlugin>>& Plugins,
 			ReadPluginsInDirectory(Root / TEXT("Plugins"), EPluginType::Project, Plugins, ChildPlugins);
 			for (auto const& Platform : FDataDrivenPlatformInfoRegistry::GetAllPlatformInfos())
 			{
+				for (auto const& AdditionalFolder : Platform.Value.AdditionalRestrictedFolders)
+				{
+					ReadPluginsInDirectory(Root / TEXT("Platforms") / AdditionalFolder / TEXT("Plugins"), EPluginType::Project, Plugins, ChildPlugins);
+				}
 				ReadPluginsInDirectory(Root / TEXT("Platforms") / Platform.Key / TEXT("Plugins"), EPluginType::Project, Plugins, ChildPlugins);
 			}
 		}
@@ -466,7 +474,7 @@ void FPluginManager::CreatePluginObject(const FString& FileName, const FPluginDe
 		Plugins[Plugin->GetName()] = Plugin;
 		UE_LOG(LogPluginManager, Verbose, TEXT("Replacing engine version of '%s' plugin with game version"), *Plugin->GetName());
 	}
-	else if( (*ExistingPlugin)->Type != EPluginType::Project || Type != EPluginType::Engine)
+	else if (((*ExistingPlugin)->Type != EPluginType::Project || Type != EPluginType::Engine) && (*ExistingPlugin)->FileName != Plugin->FileName)
 	{
 		UE_LOG(LogPluginManager, Warning, TEXT("Plugin '%s' exists at '%s' and '%s' - second location will be ignored"), *Plugin->GetName(), *(*ExistingPlugin)->FileName, *Plugin->FileName);
 	}
