Rework checks for whether a temporary target is necessary to ensure the same code is used at runtime and in UAT.

#rnx

diff --git a/Engine/Source/Developer/TargetPlatform/Private/TargetPlatformBase.cpp b/Engine/Source/Developer/TargetPlatform/Private/TargetPlatformBase.cpp
--- a/Engine/Source/Developer/TargetPlatform/Private/TargetPlatformBase.cpp
+++ b/Engine/Source/Developer/TargetPlatform/Private/TargetPlatformBase.cpp
@@ -86,6 +86,73 @@ static bool IsPluginCompiledForTarget(const IPlugin& Plugin, const FProjectDescr
 	return bCompiledForTarget;
 }
 
+static bool ConfigureEnabledPlugins(const FPluginReferenceDescriptor& FirstReference, const FProjectDescriptor* ProjectDescriptor, const FString& TargetName, const FString& Platform, EBuildConfiguration Configuration, EBuildTargetType TargetType, const TMap<FString, IPlugin*>& Plugins, TSet<FString>& EnabledPluginNames)
+{
+	if (!EnabledPluginNames.Contains(FirstReference.Name))
+	{
+		// Set of plugin names we've added to the queue for processing
+		TSet<FString> NewPluginNames;
+		NewPluginNames.Add(FirstReference.Name);
+
+		// Queue of plugin references to consider
+		TArray<const FPluginReferenceDescriptor*> NewPluginReferences;
+		NewPluginReferences.Add(&FirstReference);
+
+		// Loop through the queue of plugin references that need to be enabled, queuing more items as we go
+		TArray<TSharedRef<IPlugin>> NewPlugins;
+		for (int32 Idx = 0; Idx < NewPluginReferences.Num(); Idx++)
+		{
+			const FPluginReferenceDescriptor& Reference = *NewPluginReferences[Idx];
+
+			// Check if the plugin is required for this platform
+			if(!Reference.IsEnabledForPlatform(Platform) || !Reference.IsEnabledForTargetConfiguration(Configuration) || !Reference.IsEnabledForTarget(TargetType))
+			{
+				continue;
+			}
+
+			// Find the plugin being enabled
+			const IPlugin* const* PluginPtr = Plugins.Find(Reference.Name);
+			if (PluginPtr == nullptr)
+			{
+				continue;
+			}
+
+			// Check the plugin supports this platform
+			const FPluginDescriptor& PluginDescriptor = (*PluginPtr)->GetDescriptor();
+			if(!PluginDescriptor.SupportsTargetPlatform(Platform))
+			{
+				continue;
+			}
+
+			// Check that this plugin supports the current program
+			if (TargetType == EBuildTargetType::Program && !PluginDescriptor.SupportedPrograms.Contains(TargetName))
+			{
+				continue;
+			}
+
+			// Skip loading Enterprise plugins when project is not an Enterprise project
+			if ((*PluginPtr)->GetType() == EPluginType::Enterprise && !ProjectDescriptor->bIsEnterpriseProject)
+			{
+				continue;
+			}
+
+			// Add references to all its dependencies
+			for (const FPluginReferenceDescriptor& NextReference : PluginDescriptor.Plugins)
+			{
+				if (!EnabledPluginNames.Contains(NextReference.Name) && !NewPluginNames.Contains(NextReference.Name))
+				{
+					NewPluginNames.Add(NextReference.Name);
+					NewPluginReferences.Add(&NextReference);
+				}
+			}
+
+			// Add the plugin
+			EnabledPluginNames.Add((*PluginPtr)->GetName());
+		}
+	}
+	return true;
+}
+
 bool FTargetPlatformBase::RequiresTempTarget(bool bProjectHasCode, EBuildConfiguration Configuration, bool bRequiresAssetNativization, FText& OutReason) const
 {
 	// check to see if we already have a Target.cs file
@@ -118,33 +185,14 @@ bool FTargetPlatformBase::RequiresTempTarget(bool bProjectHasCode, EBuildConfigu
 		return true;
 	}
 
-	// find if there are any plugins enabled or disabled which differ from the default
-	for(const TSharedRef<IPlugin>& Plugin: IPluginManager::Get().GetDiscoveredPlugins())
+	// check if there's a non-default plugin change
+	FText Reason;
+	if (IPluginManager::Get().RequiresTempTargetForCodePlugin(Project, PlatformInfo->UBTTargetId.ToString(), Configuration, PlatformInfo->PlatformType, Reason))
 	{
-		const FString Platform = PlatformInfo->UBTTargetId.ToString();
-
-		bool bCompiledForTarget = IsPluginCompiledForTarget(*Plugin, Project, Platform, Configuration, PlatformInfo->PlatformType, RequiresCookedData());
-
-		bool bCompiledForBaseTarget = false;
-		if (!Plugin->GetDescriptor().bInstalled)
-		{
-			bCompiledForBaseTarget |= IsPluginCompiledForTarget(*Plugin, nullptr, Platform, Configuration, PlatformInfo->PlatformType, RequiresCookedData());
-		}
-
-		if (bCompiledForTarget != bCompiledForBaseTarget)
-		{
-			if (bCompiledForTarget)
-			{
-				OutReason = FText::Format(LOCTEXT("TempTarget_PluginEnabled", "{0} plugin is enabled"), FText::FromString(Plugin->GetName()));
-				return true;
-			}
-			else
-			{
-				OutReason = FText::Format(LOCTEXT("TempTarget_PluginDisabled", "{0} plugin is disabled"), FText::FromString(Plugin->GetName()));
-				return true;
-			}
-		}
+		OutReason = Reason;
+		return true;
 	}
+
 	return false;
 }
 
diff --git a/Engine/Source/Programs/AutomationTool/AutomationUtils/ProjectUtils.cs b/Engine/Source/Programs/AutomationTool/AutomationUtils/ProjectUtils.cs
--- a/Engine/Source/Programs/AutomationTool/AutomationUtils/ProjectUtils.cs
+++ b/Engine/Source/Programs/AutomationTool/AutomationUtils/ProjectUtils.cs
@@ -162,6 +162,9 @@ namespace AutomationTool
 			return false;
 		}
 
+		/// <summary>
+		/// NOTE: This function must mirror the functionality of TargetPlatformBase::RequiresTempTarget
+		/// </summary>
 		private static bool RequiresTempTarget(FileReference RawProjectPath, bool bProjectHasCode, UnrealTargetPlatform Platform, UnrealTargetConfiguration Configuration, TargetType TargetType, bool bRequiresAssetNativization, bool bRequiresCookedData, out string OutReason)
 		{
 			// check to see if we already have a Target.cs file
@@ -201,35 +204,225 @@ namespace AutomationTool
 			// Read the project descriptor, and find all the plugins available to this project
 			ProjectDescriptor Project = ProjectDescriptor.FromFile(RawProjectPath);
 
+			// Enumerate all the available plugins
+			Dictionary<string, PluginInfo> AllPlugins = Plugins.ReadAvailablePlugins(CommandUtils.EngineDirectory, RawProjectPath, new string[0]).ToDictionary(x => x.Name, x => x, StringComparer.OrdinalIgnoreCase);
+
 			// find if there are any plugins enabled or disabled which differ from the default
-			List<PluginInfo> Plugins = UnrealBuildTool.Plugins.ReadAvailablePlugins(CommandUtils.EngineDirectory, RawProjectPath, Project.AdditionalPluginDirectories);
-			foreach (PluginInfo Plugin in Plugins)
+			string Reason;
+			if (RequiresTempTargetForCodePlugin(Project, Platform, Configuration, TargetType, AllPlugins, out Reason))
 			{
-				bool bPluginEnabledForTarget = UnrealBuildTool.Plugins.IsPluginCompiledForTarget(Plugin, Project, Platform, Configuration, TargetType, bRequiresCookedData);
+				OutReason = Reason;
+				return true;
+			}
+
+			OutReason = null;
+			return false;
+		}
+
+		/// <summary>
+		/// NOTE: This function must mirror FPluginManager::RequiresTempTargetForCodePlugin
+		/// </summary>
+		static bool RequiresTempTargetForCodePlugin(ProjectDescriptor ProjectDescriptor, UnrealTargetPlatform Platform, UnrealTargetConfiguration Configuration, TargetType TargetType, Dictionary<string, PluginInfo> AllPlugins, out string OutReason)
+		{
+			PluginReferenceDescriptor MissingPlugin;
+
+			HashSet<string> ProjectCodePlugins = new HashSet<string>(StringComparer.OrdinalIgnoreCase);
+			if (!GetCodePluginsForTarget(ProjectDescriptor, Platform, Configuration, TargetType, ProjectCodePlugins, AllPlugins, out MissingPlugin))
+			{
+				OutReason = String.Format("{0} plugin is referenced by target but not found", MissingPlugin.Name);
+				return true;
+			}
+
+			HashSet<string> DefaultCodePlugins = new HashSet<string>(StringComparer.OrdinalIgnoreCase);
+			if (!GetCodePluginsForTarget(null, Platform, Configuration, TargetType, DefaultCodePlugins, AllPlugins, out MissingPlugin))
+			{
+				OutReason = String.Format("{0} plugin is referenced by the default target but not found", MissingPlugin.Name);
+				return true;
+			}
+
+			foreach (string ProjectCodePlugin in ProjectCodePlugins)
+			{
+				if (!DefaultCodePlugins.Contains(ProjectCodePlugin))
+				{
+					OutReason = String.Format("{0} plugin is enabled", ProjectCodePlugin);
+					return true;
+				}
+			}
 
-				bool bPluginEnabledForBaseTarget = false;
-				if(!Plugin.Descriptor.bInstalled)
+			foreach (string DefaultCodePlugin in DefaultCodePlugins)
+			{
+				if (!ProjectCodePlugins.Contains(DefaultCodePlugin))
 				{
-					bPluginEnabledForBaseTarget |= UnrealBuildTool.Plugins.IsPluginCompiledForTarget(Plugin, null, Platform, Configuration, TargetType, bRequiresCookedData);
+					OutReason = String.Format("{0} plugin is disabled", DefaultCodePlugin);
+					return true;
 				}
+			}
+
+			OutReason = null;
+			return false;
+		}
+
+		/// <summary>
+		/// NOTE: This function must mirror FPluginManager::GetCodePluginsForTarget
+		/// </summary>
+		static bool GetCodePluginsForTarget(ProjectDescriptor ProjectDescriptor, UnrealTargetPlatform Platform, UnrealTargetConfiguration Configuration, TargetType TargetType, HashSet<string> CodePluginNames, Dictionary<string, PluginInfo> AllPlugins, out PluginReferenceDescriptor OutMissingPlugin)
+		{
+			bool bLoadPluginsForTargetPlatforms = (TargetType == TargetType.Editor);
 
-				if (bPluginEnabledForTarget != bPluginEnabledForBaseTarget)
+			// Map of all enabled plugins
+			Dictionary<string, PluginInfo> EnabledPlugins = new Dictionary<string, PluginInfo>(StringComparer.OrdinalIgnoreCase);
+					   			 
+			// Keep a set of all the plugin names that have been configured. We read configuration data from different places, but only configure a plugin from the first place that it's referenced.
+			HashSet<string> ConfiguredPluginNames = new HashSet<string>(StringComparer.OrdinalIgnoreCase);
+
+			// Find all the plugin references in the project file
+			if (ProjectDescriptor != null && ProjectDescriptor.Plugins != null)
+			{
+				// Copy the plugin references, since we may modify the project if any plugins are missing
+				foreach (PluginReferenceDescriptor PluginReference in ProjectDescriptor.Plugins)
 				{
-					if(bPluginEnabledForTarget)
+					if(!ConfiguredPluginNames.Contains(PluginReference.Name))
 					{
-						OutReason = String.Format("{0} plugin is enabled", Plugin.Name);
-						return true;
+						PluginReferenceDescriptor MissingPlugin;
+						if (!ConfigureEnabledPluginForTarget(PluginReference, ProjectDescriptor, null, Platform, Configuration, TargetType, bLoadPluginsForTargetPlatforms, AllPlugins, EnabledPlugins, out MissingPlugin))
+						{
+							OutMissingPlugin = MissingPlugin;
+							return false;
+						}
+						ConfiguredPluginNames.Add(PluginReference.Name);
 					}
-					else
+				}
+			}
+
+			// Add the plugins which are enabled by default
+			foreach (KeyValuePair<string, PluginInfo> PluginPair in AllPlugins)
+			{
+				if(PluginPair.Value.EnabledByDefault && !ConfiguredPluginNames.Contains(PluginPair.Key))
+				{
+					PluginReferenceDescriptor MissingPlugin;
+					if (!ConfigureEnabledPluginForTarget(new PluginReferenceDescriptor(PluginPair.Key, null, true), ProjectDescriptor, null, Platform, Configuration, TargetType, bLoadPluginsForTargetPlatforms, AllPlugins, EnabledPlugins, out MissingPlugin))
 					{
-						OutReason = String.Format("{0} plugin is disabled", Plugin.Name);
-						return true;
+						OutMissingPlugin = MissingPlugin;
+						return false;
 					}
+					ConfiguredPluginNames.Add(PluginPair.Key);
 				}
 			}
 
-			OutReason = null;
-			return false;
+			// Figure out which plugins have code 
+			bool bBuildDeveloperTools = (TargetType == TargetType.Editor || TargetType == TargetType.Program || (Configuration != UnrealTargetConfiguration.Test && Configuration != UnrealTargetConfiguration.Shipping));
+			bool bRequiresCookedData = (TargetType != TargetType.Editor);
+			foreach (KeyValuePair<string, PluginInfo> Pair in EnabledPlugins)
+			{
+				ModuleDescriptor[] Modules = Pair.Value.Descriptor.Modules;
+				if (Modules != null)
+				{
+					foreach (ModuleDescriptor Module in Pair.Value.Descriptor.Modules)
+					{
+						if (Module.IsCompiledInConfiguration(Platform, Configuration, null, TargetType, bBuildDeveloperTools, bRequiresCookedData))
+						{
+							CodePluginNames.Add(Pair.Key);
+							break;
+						}
+					}
+				}
+			}
+
+			OutMissingPlugin = null;
+			return true;
+		}
+
+		/// <summary>
+		/// NOTE: This function should mirror FPluginManager::ConfigureEnabledPluginForTarget
+		/// </summary>
+		static bool ConfigureEnabledPluginForTarget(PluginReferenceDescriptor FirstReference, ProjectDescriptor ProjectDescriptor, string TargetName, UnrealTargetPlatform Platform, UnrealTargetConfiguration Configuration, TargetType TargetType, bool bLoadPluginsForTargetPlatforms, Dictionary<string, PluginInfo> AllPlugins, Dictionary<string, PluginInfo> EnabledPlugins, out PluginReferenceDescriptor OutMissingPlugin)
+		{
+			if (!EnabledPlugins.ContainsKey(FirstReference.Name))
+			{
+				// Set of plugin names we've added to the queue for processing
+				HashSet<string> NewPluginNames = new HashSet<string>(StringComparer.OrdinalIgnoreCase);
+				NewPluginNames.Add(FirstReference.Name);
+
+				// Queue of plugin references to consider
+				List<PluginReferenceDescriptor> NewPluginReferences = new List<PluginReferenceDescriptor>();
+				NewPluginReferences.Add(FirstReference);
+
+				// Loop through the queue of plugin references that need to be enabled, queuing more items as we go
+				for (int Idx = 0; Idx < NewPluginReferences.Count; Idx++)
+				{
+					PluginReferenceDescriptor Reference = NewPluginReferences[Idx];
+
+					// Check if the plugin is required for this platform
+					if(!Reference.IsEnabledForPlatform(Platform) || !Reference.IsEnabledForTargetConfiguration(Configuration) || !Reference.IsEnabledForTarget(TargetType))
+					{
+						Log.TraceLog("Ignoring plugin '{0}' for platform/configuration", Reference.Name);
+						continue;
+					}
+
+					// Check if the plugin is required for this platform
+					if(!bLoadPluginsForTargetPlatforms && !Reference.IsSupportedTargetPlatform(Platform))
+					{
+						Log.TraceLog("Ignoring plugin '{0}' due to unsupported platform", Reference.Name);
+						continue;
+					}
+
+					// Find the plugin being enabled
+					PluginInfo Plugin;
+					if (!AllPlugins.TryGetValue(Reference.Name, out Plugin))
+					{
+						// Ignore any optional plugins
+						if (Reference.bOptional)
+						{
+							Log.TraceLog("Ignored optional reference to '%s' plugin; plugin was not found.", Reference.Name);
+							continue;
+						}
+
+						// Add it to the missing list
+						OutMissingPlugin = Reference;
+						return false;
+					}
+
+					// Check the plugin supports this platform
+					if(!bLoadPluginsForTargetPlatforms && !Plugin.Descriptor.SupportsTargetPlatform(Platform))
+					{
+						Log.TraceLog("Ignoring plugin '{0}' due to unsupported platform in plugin descriptor", Reference.Name);
+						continue;
+					}
+
+					// Check that this plugin supports the current program
+					if (TargetType == TargetType.Program && !Plugin.Descriptor.SupportedPrograms.Contains(TargetName))
+					{
+						Log.TraceLog("Ignoring plugin '{0}' due to absence from the supported programs list", Reference.Name);
+						continue;
+					}
+
+					// Skip loading Enterprise plugins when project is not an Enterprise project
+					if (Plugin.Type == PluginType.Enterprise && ProjectDescriptor != null && !ProjectDescriptor.IsEnterpriseProject)
+					{
+						Log.TraceLog("Ignoring plugin '{0}' due to not being an enterprise project", Reference.Name);
+						continue;
+					}
+
+					// Add references to all its dependencies
+					if (Plugin.Descriptor.Plugins != null)
+					{
+						foreach (PluginReferenceDescriptor NextReference in Plugin.Descriptor.Plugins)
+						{
+							if (!EnabledPlugins.ContainsKey(NextReference.Name) && !NewPluginNames.Contains(NextReference.Name))
+							{
+								NewPluginNames.Add(NextReference.Name);
+								NewPluginReferences.Add(NextReference);
+							}
+						}
+					}
+
+					// Add the plugin
+					EnabledPlugins.Add(Plugin.Name, Plugin);
+				}
+			}
+
+			OutMissingPlugin = null;
+			return true;
 		}
 
 		private static void GenerateTempTarget(FileReference RawProjectPath)
diff --git a/Engine/Source/Runtime/Projects/Private/PluginManager.cpp b/Engine/Source/Runtime/Projects/Private/PluginManager.cpp
--- a/Engine/Source/Runtime/Projects/Private/PluginManager.cpp
+++ b/Engine/Source/Runtime/Projects/Private/PluginManager.cpp
@@ -512,7 +512,7 @@ bool FPluginManager::ConfigureEnabledPlugins()
 		bHaveConfiguredEnabledPlugins = true;
 
 		// Set of all the plugins which have been enabled
-		TSet<FString> EnabledPluginNames;
+		TMap<FString, FPlugin*> EnabledPlugins;
 
 		// Keep a set of all the plugin names that have been configured. We read configuration data from different places, but only configure a plugin from the first place that it's referenced.
 		TSet<FString> ConfiguredPluginNames;
@@ -545,7 +545,7 @@ bool FPluginManager::ConfigureEnabledPlugins()
 				{
 					if (!ConfiguredPluginNames.Contains(EnablePluginName) && !ExceptPlugins.Contains(EnablePluginName))
 					{
-						if (!ConfigureEnabledPlugin(FPluginReferenceDescriptor(EnablePluginName, true), EnabledPluginNames))
+						if (!ConfigureEnabledPluginForCurrentTarget(FPluginReferenceDescriptor(EnablePluginName, true), EnabledPlugins))
 						{
 							return false;
 						}
@@ -563,7 +563,7 @@ bool FPluginManager::ConfigureEnabledPlugins()
 			{
 				if (!ConfiguredPluginNames.Contains(TargetEnabledPlugin))
 				{
-					if (!ConfigureEnabledPlugin(FPluginReferenceDescriptor(TargetEnabledPlugin, true), EnabledPluginNames))
+					if (!ConfigureEnabledPluginForCurrentTarget(FPluginReferenceDescriptor(TargetEnabledPlugin, true), EnabledPlugins))
 					{
 						return false;
 					}
@@ -577,7 +577,7 @@ bool FPluginManager::ConfigureEnabledPlugins()
 			{
 				if (!ConfiguredPluginNames.Contains(TargetDisabledPlugin))
 				{
-					if (!ConfigureEnabledPlugin(FPluginReferenceDescriptor(TargetDisabledPlugin, false), EnabledPluginNames))
+					if (!ConfigureEnabledPluginForCurrentTarget(FPluginReferenceDescriptor(TargetDisabledPlugin, false), EnabledPlugins))
 					{
 						return false;
 					}
@@ -595,7 +595,7 @@ bool FPluginManager::ConfigureEnabledPlugins()
 				{
 					if(!ConfiguredPluginNames.Contains(PluginReference.Name))
 					{
-						if (!ConfigureEnabledPlugin(PluginReference, EnabledPluginNames))
+						if (!ConfigureEnabledPluginForCurrentTarget(PluginReference, EnabledPlugins))
 						{
 							return false;
 						}
@@ -609,7 +609,7 @@ bool FPluginManager::ConfigureEnabledPlugins()
 			{
 				if(PluginPair.Value->IsEnabledByDefault() && !ConfiguredPluginNames.Contains(PluginPair.Key))
 				{
-					if (!ConfigureEnabledPlugin(FPluginReferenceDescriptor(PluginPair.Key, true), EnabledPluginNames))
+					if (!ConfigureEnabledPluginForCurrentTarget(FPluginReferenceDescriptor(PluginPair.Key, true), EnabledPlugins))
 					{
 						return false;
 					}
@@ -626,7 +626,7 @@ bool FPluginManager::ConfigureEnabledPlugins()
 		{
 			if(!ConfiguredPluginNames.Contains(PluginName))
 			{
-				if (!ConfigureEnabledPlugin(FPluginReferenceDescriptor(PluginName, true), EnabledPluginNames))
+				if (!ConfigureEnabledPluginForCurrentTarget(FPluginReferenceDescriptor(PluginName, true), EnabledPlugins))
 				{
 					return false;
 				}
@@ -635,6 +635,30 @@ bool FPluginManager::ConfigureEnabledPlugins()
 		}
 #endif
 
+		// Mark all the plugins as enabled
+		for (TPair<FString, FPlugin*>& Pair : EnabledPlugins)
+		{
+			FPlugin& Plugin = *Pair.Value;
+
+#if !IS_MONOLITHIC
+			// Mount the binaries directory, and check the modules are valid
+			if (Plugin.Descriptor.Modules.Num() > 0)
+			{
+				// Mount the binaries directory
+				const FString PluginBinariesPath = FPaths::Combine(*FPaths::GetPath(Plugin.FileName), TEXT("Binaries"), FPlatformProcess::GetBinariesSubdirectory());
+				FModuleManager::Get().AddBinariesDirectory(*PluginBinariesPath, Plugin.GetLoadedFrom() == EPluginLoadedFrom::Project);
+			}
+
+			// Check the declared engine version. This is a soft requirement, so allow the user to skip over it.
+			if (!IsPluginCompatible(Plugin) && !PromptToLoadIncompatiblePlugin(Plugin))
+			{
+				UE_LOG(LogPluginManager, Display, TEXT("Skipping load of '%s'."), *Plugin.Name);
+				continue;
+			}
+#endif
+			Plugin.bEnabled = true;
+		}
+
 		// If we made it here, we have all the required plugins
 		bHaveAllRequiredPlugins = true;
 
@@ -751,9 +775,146 @@ bool FPluginManager::ConfigureEnabledPlugins()
 	return bHaveAllRequiredPlugins;
 }
 
-bool FPluginManager::ConfigureEnabledPlugin(const FPluginReferenceDescriptor& FirstReference, TSet<FString>& EnabledPluginNames)
+bool FPluginManager::RequiresTempTargetForCodePlugin(const FProjectDescriptor* ProjectDescriptor, const FString& Platform, EBuildConfiguration Configuration, EBuildTargetType TargetType, FText& OutReason)
+{
+	const FPluginReferenceDescriptor* MissingPlugin;
+
+	TSet<FString> ProjectCodePlugins;
+	if (!GetCodePluginsForProject(ProjectDescriptor, Platform, Configuration, TargetType, AllPlugins, ProjectCodePlugins, MissingPlugin))
+	{
+		OutReason = FText::Format(LOCTEXT("TempTarget_MissingPluginForTarget", "{0} plugin is referenced by target but not found"), FText::FromString(MissingPlugin->Name));
+		return true;
+	}
+
+	TSet<FString> DefaultCodePlugins;
+	if (!GetCodePluginsForProject(nullptr, Platform, Configuration, TargetType, AllPlugins, DefaultCodePlugins, MissingPlugin))
+	{
+		OutReason = FText::Format(LOCTEXT("TempTarget_MissingPluginForTarget", "{0} plugin is referenced by the default target but not found"), FText::FromString(MissingPlugin->Name));
+		return true;
+	}
+
+	for (const FString& ProjectCodePlugin : ProjectCodePlugins)
+	{
+		if (!DefaultCodePlugins.Contains(ProjectCodePlugin))
+		{
+			OutReason = FText::Format(LOCTEXT("TempTarget_PluginEnabled", "{0} plugin is enabled"), FText::FromString(ProjectCodePlugin));
+			return true;
+		}
+	}
+
+	for (const FString& DefaultCodePlugin : DefaultCodePlugins)
+	{
+		if (!ProjectCodePlugins.Contains(DefaultCodePlugin))
+		{
+			OutReason = FText::Format(LOCTEXT("TempTarget_PluginDisabled", "{0} plugin is disabled"), FText::FromString(DefaultCodePlugin));
+			return true;
+		}
+	}
+
+	return false;
+}
+
+bool FPluginManager::GetCodePluginsForProject(const FProjectDescriptor* ProjectDescriptor, const FString& Platform, EBuildConfiguration Configuration, EBuildTargetType TargetType, const TMap<FString, TSharedRef<FPlugin>>& AllPlugins, TSet<FString>& CodePluginNames, const FPluginReferenceDescriptor*& OutMissingPlugin)
+{
+	// Can only check the current project at the moment, since we won't have enumerated them otherwise
+	check(ProjectDescriptor == nullptr || ProjectDescriptor == IProjectManager::Get().GetCurrentProject());
+
+	// Always false for content-only projects
+	const bool bLoadPluginsForTargetPlatforms = (TargetType == EBuildTargetType::Editor);
+
+	// Map of all enabled plugins
+	TMap<FString, FPlugin*> EnabledPlugins;
+
+	// Keep a set of all the plugin names that have been configured. We read configuration data from different places, but only configure a plugin from the first place that it's referenced.
+	TSet<FString> ConfiguredPluginNames;
+
+	// Find all the plugin references in the project file
+	if (ProjectDescriptor != nullptr)
+	{
+		// Copy the plugin references, since we may modify the project if any plugins are missing
+		TArray<FPluginReferenceDescriptor> PluginReferences(ProjectDescriptor->Plugins);
+		for (const FPluginReferenceDescriptor& PluginReference : PluginReferences)
+		{
+			if(!ConfiguredPluginNames.Contains(PluginReference.Name))
+			{
+				if (!ConfigureEnabledPluginForTarget(PluginReference, ProjectDescriptor, FString(), Platform, Configuration, TargetType, bLoadPluginsForTargetPlatforms, AllPlugins, EnabledPlugins, OutMissingPlugin))
+				{
+					return false;
+				}
+				ConfiguredPluginNames.Add(PluginReference.Name);
+			}
+		}
+	}
+
+	// Add the plugins which are enabled by default
+	for (const TPair<FString, TSharedRef<FPlugin>>& PluginPair : AllPlugins)
+	{
+		if(PluginPair.Value->IsEnabledByDefault() && !ConfiguredPluginNames.Contains(PluginPair.Key))
+		{
+			if (!ConfigureEnabledPluginForTarget(FPluginReferenceDescriptor(PluginPair.Key, true), ProjectDescriptor, FString(), Platform, Configuration, TargetType, bLoadPluginsForTargetPlatforms, AllPlugins, EnabledPlugins, OutMissingPlugin))
+			{
+				return false;
+			}
+			ConfiguredPluginNames.Add(PluginPair.Key);
+		}
+	}
+
+	// Figure out which plugins have code 
+	bool bBuildDeveloperTools = (TargetType == EBuildTargetType::Editor || TargetType == EBuildTargetType::Program || (Configuration != EBuildConfiguration::Test && Configuration != EBuildConfiguration::Shipping));
+	bool bRequiresCookedData = (TargetType != EBuildTargetType::Editor);
+	for (const TPair<FString, FPlugin*>& Pair : EnabledPlugins)
+	{
+		for (const FModuleDescriptor& Module : Pair.Value->GetDescriptor().Modules)
+		{
+			if (Module.IsCompiledInConfiguration(Platform, Configuration, FString(), TargetType, bBuildDeveloperTools, bRequiresCookedData))
+			{
+				CodePluginNames.Add(Pair.Key);
+				break;
+			}
+		}
+	}
+
+	return true;
+}
+
+bool FPluginManager::ConfigureEnabledPluginForCurrentTarget(const FPluginReferenceDescriptor& FirstReference, TMap<FString, FPlugin*>& EnabledPlugins)
 {
-	if (!EnabledPluginNames.Contains(FirstReference.Name))
+	const FPluginReferenceDescriptor* MissingPlugin;
+	if (!ConfigureEnabledPluginForTarget(FirstReference, IProjectManager::Get().GetCurrentProject(), UE_APP_NAME, FPlatformMisc::GetUBTPlatform(), FApp::GetBuildConfiguration(), FApp::GetBuildTargetType(), (bool)LOAD_PLUGINS_FOR_TARGET_PLATFORMS, AllPlugins, EnabledPlugins, MissingPlugin))
+	{
+		// If we're in unattended mode, don't open any windows and fatal out
+		if (FApp::IsUnattended())
+		{
+			UE_LOG(LogPluginManager, Fatal, TEXT("This project requires the '%s' plugin. Install it and try again, or remove it from the project's required plugin list."), *MissingPlugin->Name);
+			return false;
+		}
+
+#if !IS_MONOLITHIC
+		// Try to download it from the marketplace
+		if (MissingPlugin->MarketplaceURL.Len() > 0 && PromptToDownloadPlugin(MissingPlugin->Name, MissingPlugin->MarketplaceURL))
+		{
+			UE_LOG(LogPluginManager, Display, TEXT("Downloading '%s' plugin from marketplace (%s)."), *MissingPlugin->Name, *MissingPlugin->MarketplaceURL);
+			return false;
+		}
+
+		// Prompt to disable it in the project file, if possible
+		if (PromptToDisableMissingPlugin(FirstReference.Name, MissingPlugin->Name))
+		{
+			UE_LOG(LogPluginManager, Display, TEXT("Disabled plugin '%s', continuing."), *FirstReference.Name);
+			return true;
+		}
+#endif
+
+		// Unable to continue
+		UE_LOG(LogPluginManager, Error, TEXT("Unable to load plugin '%s'. Aborting."), *MissingPlugin->Name);
+		return false;
+	}
+	return true;
+}
+
+bool FPluginManager::ConfigureEnabledPluginForTarget(const FPluginReferenceDescriptor& FirstReference, const FProjectDescriptor* ProjectDescriptor, const FString& TargetName, const FString& Platform, EBuildConfiguration Configuration, EBuildTargetType TargetType, bool bLoadPluginsForTargetPlatforms, const TMap<FString, TSharedRef<FPlugin>>& AllPlugins, TMap<FString, FPlugin*>& EnabledPlugins, const FPluginReferenceDescriptor*& OutMissingPlugin)
+{
+	if (!EnabledPlugins.Contains(FirstReference.Name))
 	{
 		// Set of plugin names we've added to the queue for processing
 		TSet<FString> NewPluginNames;
@@ -770,23 +931,21 @@ bool FPluginManager::ConfigureEnabledPlugin(const FPluginReferenceDescriptor& Fi
 			const FPluginReferenceDescriptor& Reference = *NewPluginReferences[Idx];
 
 			// Check if the plugin is required for this platform
-			if(!Reference.IsEnabledForPlatform(FPlatformMisc::GetUBTPlatform()) || !Reference.IsEnabledForTargetConfiguration(FApp::GetBuildConfiguration()) || !Reference.IsEnabledForTarget(FApp::GetBuildTargetType()))
+			if(!Reference.IsEnabledForPlatform(Platform) || !Reference.IsEnabledForTargetConfiguration(Configuration) || !Reference.IsEnabledForTarget(TargetType))
 			{
 				UE_LOG(LogPluginManager, Verbose, TEXT("Ignoring plugin '%s' for platform/configuration"), *Reference.Name);
 				continue;
 			}
 
 			// Check if the plugin is required for this platform
-#if !LOAD_PLUGINS_FOR_TARGET_PLATFORMS
-			if(!Reference.IsSupportedTargetPlatform(FPlatformMisc::GetUBTPlatform()))
+			if(!bLoadPluginsForTargetPlatforms && !Reference.IsSupportedTargetPlatform(Platform))
 			{
 				UE_LOG(LogPluginManager, Verbose, TEXT("Ignoring plugin '%s' due to unsupported platform"), *Reference.Name);
 				continue;
 			}
-#endif
 
 			// Find the plugin being enabled
-			TSharedRef<FPlugin>* PluginPtr = AllPlugins.Find(Reference.Name);
+			const TSharedRef<FPlugin>* PluginPtr = AllPlugins.Find(Reference.Name);
 			if (PluginPtr == nullptr)
 			{
 				// Ignore any optional plugins
@@ -796,31 +955,8 @@ bool FPluginManager::ConfigureEnabledPlugin(const FPluginReferenceDescriptor& Fi
 					continue;
 				}
 
-				// If we're in unattended mode, don't open any windows and fatal out
-				if (FApp::IsUnattended())
-				{
-					UE_LOG(LogPluginManager, Fatal, TEXT("This project requires the '%s' plugin. Install it and try again, or remove it from the project's required plugin list."), *Reference.Name);
-					return false;
-				}
-
-#if !IS_MONOLITHIC
-				// Try to download it from the marketplace
-				if (Reference.MarketplaceURL.Len() > 0 && PromptToDownloadPlugin(Reference.Name, Reference.MarketplaceURL))
-				{
-					UE_LOG(LogPluginManager, Display, TEXT("Downloading '%s' plugin from marketplace (%s)."), *Reference.Name, *Reference.MarketplaceURL);
-					return false;
-				}
-
-				// Prompt to disable it in the project file, if possible
-				if (PromptToDisableMissingPlugin(FirstReference.Name, Reference.Name))
-				{
-					UE_LOG(LogPluginManager, Display, TEXT("Disabled plugin '%s', continuing."), *FirstReference.Name);
-					return true;
-				}
-#endif
-
-				// Unable to continue
-				UE_LOG(LogPluginManager, Error, TEXT("Unable to load plugin '%s'. Aborting."), *Reference.Name);
+				// Add it to the missing list
+				OutMissingPlugin = &Reference;
 				return false;
 			}
 
@@ -835,49 +971,30 @@ bool FPluginManager::ConfigureEnabledPlugin(const FPluginReferenceDescriptor& Fi
 			}
 
 			// Check the plugin supports this platform
-#if !LOAD_PLUGINS_FOR_TARGET_PLATFORMS
-			if(!Plugin.Descriptor.SupportsTargetPlatform(FPlatformMisc::GetUBTPlatform()))
+			if(!bLoadPluginsForTargetPlatforms && !Plugin.Descriptor.SupportsTargetPlatform(Platform))
 			{
 				UE_LOG(LogPluginManager, Verbose, TEXT("Ignoring plugin '%s' due to unsupported platform in plugin descriptor"), *Reference.Name);
 				continue;
 			}
-#endif
+
 			// Check that this plugin supports the current program
-#if IS_PROGRAM
-			if (!Plugin.Descriptor.SupportedPrograms.Contains(UE_APP_NAME))
+			if (TargetType == EBuildTargetType::Program && !Plugin.Descriptor.SupportedPrograms.Contains(TargetName))
 			{
 				UE_LOG(LogPluginManager, Verbose, TEXT("Ignoring plugin '%s' due to absence from the supported programs list"), *Reference.Name);
 				continue;
 			}
-#endif
 
 			// Skip loading Enterprise plugins when project is not an Enterprise project
-			if (Plugin.Type == EPluginType::Enterprise && !IProjectManager::Get().IsEnterpriseProject())
+			if (Plugin.Type == EPluginType::Enterprise && ProjectDescriptor != nullptr && !ProjectDescriptor->bIsEnterpriseProject)
 			{
+				UE_LOG(LogPluginManager, Verbose, TEXT("Ignoring plugin '%s' due to not being an Enterpise project"), *Reference.Name);
 				continue;
 			}
 
-#if !IS_MONOLITHIC
-			// Mount the binaries directory, and check the modules are valid
-			if (Plugin.Descriptor.Modules.Num() > 0)
-			{
-				// Mount the binaries directory
-				const FString PluginBinariesPath = FPaths::Combine(*FPaths::GetPath(Plugin.FileName), TEXT("Binaries"), FPlatformProcess::GetBinariesSubdirectory());
-				FModuleManager::Get().AddBinariesDirectory(*PluginBinariesPath, Plugin.GetLoadedFrom() == EPluginLoadedFrom::Project);
-			}
-
-			// Check the declared engine version. This is a soft requirement, so allow the user to skip over it.
-			if (!IsPluginCompatible(Plugin) && !PromptToLoadIncompatiblePlugin(Plugin, FirstReference.Name))
-			{
-				UE_LOG(LogPluginManager, Display, TEXT("Skipping load of '%s'."), *Plugin.Name);
-				return true;
-			}
-#endif
-
 			// Add references to all its dependencies
 			for (const FPluginReferenceDescriptor& NextReference : Plugin.Descriptor.Plugins)
 			{
-				if (!EnabledPluginNames.Contains(NextReference.Name) && !NewPluginNames.Contains(NextReference.Name))
+				if (!EnabledPlugins.Contains(NextReference.Name) && !NewPluginNames.Contains(NextReference.Name))
 				{
 					NewPluginNames.Add(NextReference.Name);
 					NewPluginReferences.Add(&NextReference);
@@ -885,14 +1002,7 @@ bool FPluginManager::ConfigureEnabledPlugin(const FPluginReferenceDescriptor& Fi
 			}
 
 			// Add the plugin
-			NewPlugins.Add(*PluginPtr);
-		}
-
-		// Mark all the plugins as enabled
-		for (TSharedRef<FPlugin>& NewPlugin : NewPlugins)
-		{
-			NewPlugin->bEnabled = true;
-			EnabledPluginNames.Add(NewPlugin->Name);
+			EnabledPlugins.Add(Plugin.GetName(), &Plugin);
 		}
 	}
 	return true;
@@ -986,19 +1096,10 @@ bool FPluginManager::IsPluginCompatible(const FPlugin& Plugin)
 	return true;
 }
 
-bool FPluginManager::PromptToLoadIncompatiblePlugin(const FPlugin& Plugin, const FString& ReferencingPluginName)
+bool FPluginManager::PromptToLoadIncompatiblePlugin(const FPlugin& Plugin)
 {
 	// Format the message dependning on whether the plugin is referenced directly, or as a dependency
-	FText Message;
-	if (Plugin.Name == ReferencingPluginName)
-	{
-		Message = FText::Format(LOCTEXT("LoadIncompatiblePlugin", "The '{0}' plugin was designed for build {1}. Attempt to load it anyway?"), FText::FromString(Plugin.Name), FText::FromString(Plugin.Descriptor.EngineVersion));
-	}
-	else
-	{
-		Message = FText::Format(LOCTEXT("LoadIncompatibleDependencyPlugin", "The '{0}' plugin is required by the '{1}' plugin, but was designed for build {2}. Attempt to load it anyway?"), FText::FromString(Plugin.Name), FText::FromString(ReferencingPluginName), FText::FromString(Plugin.Descriptor.EngineVersion));
-	}
-
+	FText Message = FText::Format(LOCTEXT("LoadIncompatiblePlugin", "The '{0}' plugin was designed for build {1}. Attempt to load it anyway?"), FText::FromString(Plugin.Name), FText::FromString(Plugin.Descriptor.EngineVersion));
 	FText Caption = FText::Format(LOCTEXT("IncompatiblePluginCaption", "'{0}' is Incompatible"), FText::FromString(Plugin.Name));
 	return FMessageDialog::Open(EAppMsgType::YesNo, Message, &Caption) == EAppReturnType::Yes;
 }
diff --git a/Engine/Source/Runtime/Projects/Private/PluginManager.h b/Engine/Source/Runtime/Projects/Private/PluginManager.h
--- a/Engine/Source/Runtime/Projects/Private/PluginManager.h
+++ b/Engine/Source/Runtime/Projects/Private/PluginManager.h
@@ -6,6 +6,8 @@
 #include "PluginDescriptor.h"
 #include "Interfaces/IPluginManager.h"
 
+struct FProjectDescriptor;
+
 /**
  * Instance of a plugin in memory
  */
@@ -112,7 +114,7 @@ public:
 	virtual FNewPluginMountedEvent& OnNewPluginMounted() override;
 	virtual void MountNewlyCreatedPlugin(const FString& PluginName) override;
 	virtual FName PackageNameFromModuleName(FName ModuleName) override;
-
+	virtual bool RequiresTempTargetForCodePlugin(const FProjectDescriptor* ProjectDescriptor, const FString& Platform, EBuildConfiguration Configuration, EBuildTargetType TargetType, FText& OutReason) override;
 
 private:
 
@@ -135,11 +137,17 @@ private:
 	/** Finds all the plugin manifests in a given directory */
 	static void FindPluginManifestsInDirectory(const FString& PluginManifestDirectory, TArray<FString>& FileNames);
 
+	/** Gets all the code plugins that are enabled for a content only project */
+	static bool GetCodePluginsForProject(const FProjectDescriptor* ProjectDescriptor, const FString& Platform, EBuildConfiguration Configuration, EBuildTargetType TargetType, const TMap<FString, TSharedRef<FPlugin>>& AllPlugins, TSet<FString>& CodePluginNames, const FPluginReferenceDescriptor*& OutMissingPlugin);
+
 	/** Sets the bPluginEnabled flag on all plugins found from DiscoverAllPlugins that are enabled in config */
 	bool ConfigureEnabledPlugins();
 
 	/** Adds a single enabled plugin, and all its dependencies */
-	bool ConfigureEnabledPlugin(const FPluginReferenceDescriptor& FirstReference, TSet<FString>& EnabledPluginNames);
+	bool ConfigureEnabledPluginForCurrentTarget(const FPluginReferenceDescriptor& FirstReference, TMap<FString, FPlugin*>& EnabledPlugins);
+
+	/** Adds a single enabled plugin and all its dependencies. */
+	static bool ConfigureEnabledPluginForTarget(const FPluginReferenceDescriptor& FirstReference, const FProjectDescriptor* ProjectDescriptor, const FString& TargetName, const FString& Platform, EBuildConfiguration Configuration, EBuildTargetType TargetType, bool bLoadPluginsForTargetPlatforms, const TMap<FString, TSharedRef<FPlugin>>& AllPlugins, TMap<FString, FPlugin*>& EnabledPlugins, const FPluginReferenceDescriptor*& OutMissingPlugin);
 
 	/** Prompts the user to download a missing plugin from the given URL */
 	static bool PromptToDownloadPlugin(const FString& PluginName, const FString& MarketplaceURL);
@@ -157,7 +165,7 @@ private:
 	static bool IsPluginCompatible(const FPlugin& Plugin);
 
 	/** Prompts the user to disable a plugin */
-	static bool PromptToLoadIncompatiblePlugin(const FPlugin& Plugin, const FString& ReferencingPluginName);
+	static bool PromptToLoadIncompatiblePlugin(const FPlugin& Plugin);
 
 	/** Gets the instance of a given plugin */
 	TSharedPtr<FPlugin> FindPluginInstance(const FString& Name);
diff --git a/Engine/Source/Runtime/Projects/Public/Interfaces/IPluginManager.h b/Engine/Source/Runtime/Projects/Public/Interfaces/IPluginManager.h
--- a/Engine/Source/Runtime/Projects/Public/Interfaces/IPluginManager.h
+++ b/Engine/Source/Runtime/Projects/Public/Interfaces/IPluginManager.h
@@ -5,6 +5,8 @@
 #include "CoreMinimal.h"
 #include "PluginDescriptor.h"
 
+struct FProjectDescriptor;
+
 /**
  * Enum for where a plugin is loaded from
  */
@@ -294,6 +296,18 @@ public:
 	*/
 	virtual FName PackageNameFromModuleName(FName ModuleName) = 0;
 
+	/**
+	 * Determines if a content-only project requires a temporary target due to having a plugin enabled
+	 *
+	 * @param ProjectDescriptor The project being built
+	 * @param Platform The platform the target is being built for
+	 * @param Configuration The configuration being built
+	 * @param TargetType The type of target being built
+	 * @param OutReason If a temporary target is required, receives a message indicating why
+	 * @return True if the project requires a temp target to be generated
+	 */
+	virtual bool RequiresTempTargetForCodePlugin(const FProjectDescriptor* ProjectDescriptor, const FString& Platform, EBuildConfiguration Configuration, EBuildTargetType TargetType, FText& OutReason) = 0;
+
 public:
 
 	/**
