Fix localization conflict

#rnx

diff --git a/Engine/Source/Runtime/Projects/Private/PluginManager.cpp b/Engine/Source/Runtime/Projects/Private/PluginManager.cpp
--- a/Engine/Source/Runtime/Projects/Private/PluginManager.cpp
+++ b/Engine/Source/Runtime/Projects/Private/PluginManager.cpp
@@ -789,7 +789,7 @@ bool FPluginManager::RequiresTempTargetForCodePlugin(const FProjectDescriptor* P
 	TSet<FString> DefaultCodePlugins;
 	if (!GetCodePluginsForProject(nullptr, Platform, Configuration, TargetType, AllPlugins, DefaultCodePlugins, MissingPlugin))
 	{
-		OutReason = FText::Format(LOCTEXT("TempTarget_MissingPluginForTarget", "{0} plugin is referenced by the default target but not found"), FText::FromString(MissingPlugin->Name));
+		OutReason = FText::Format(LOCTEXT("TempTarget_MissingPluginForDefaultTarget", "{0} plugin is referenced by the default target but not found"), FText::FromString(MissingPlugin->Name));
 		return true;
 	}
 
