Don't try and load plugin manifests in editor builds

diff --git a/Engine/Source/Runtime/Projects/Private/PluginManager.cpp b/Engine/Source/Runtime/Projects/Private/PluginManager.cpp
--- a/Engine/Source/Runtime/Projects/Private/PluginManager.cpp
+++ b/Engine/Source/Runtime/Projects/Private/PluginManager.cpp
@@ -215,10 +215,12 @@ void FPluginManager::ReadAllPlugins(TMap<FString, TSharedRef<FPlugin>>& Plugins,
 
 	// Find any plugin manifest files. These give us the plugin list (and their descriptors) without needing to scour the directory tree.
 	TArray<FString> ManifestFileNames;
+#if !WITH_EDITOR
 	if (Project != nullptr)
 	{
 		FindPluginManifestsInDirectory(*FPaths::ProjectPluginsDir(), ManifestFileNames);
 	}
+#endif // !WITH_EDITOR
 
 	// track child plugins that don't want to go into main plugin set
 	TArray<TSharedRef<FPlugin>> ChildPlugins;
