#UE4 Added bExplicitlyLoaded to plugin descriptors. When set to true, the plugin's  content and code will not be mounted/loaded automatically at startup. Instead, it will be mounted when MountExplicitlyLoadedPlugin is called from project-specific code.

diff --git a/Engine/Source/Runtime/Projects/Private/PluginDescriptor.cpp b/Engine/Source/Runtime/Projects/Private/PluginDescriptor.cpp
--- a/Engine/Source/Runtime/Projects/Private/PluginDescriptor.cpp
+++ b/Engine/Source/Runtime/Projects/Private/PluginDescriptor.cpp
@@ -34,6 +34,7 @@ FPluginDescriptor::FPluginDescriptor()
 	, bInstalled(false)
 	, bRequiresBuildPlatform(false)
 	, bIsHidden(false)
+	, bExplicitlyLoaded(false)
 	, bIsPluginExtension(false)
 {
 }
@@ -142,6 +143,7 @@ bool FPluginDescriptor::Read(const FJsonObject& Object, FText& OutFailReason)
 	Object.TryGetBoolField(TEXT("Installed"), bInstalled);
 	Object.TryGetBoolField(TEXT("RequiresBuildPlatform"), bRequiresBuildPlatform);
 	Object.TryGetBoolField(TEXT("Hidden"), bIsHidden);
+	Object.TryGetBoolField(TEXT("ExplicitlyLoaded"), bExplicitlyLoaded);
 
 	bool bCanBeUsedWithUnrealHeaderTool;
 	if(Object.TryGetBoolField("CanBeUsedWithUnrealHeaderTool", bCanBeUsedWithUnrealHeaderTool) && bCanBeUsedWithUnrealHeaderTool)
@@ -239,6 +241,11 @@ void FPluginDescriptor::Write(TJsonWriter<>& Writer) const
 		Writer.WriteValue(TEXT("Hidden"), bIsHidden);
 	}
 
+	if (bExplicitlyLoaded)
+	{
+		Writer.WriteValue(TEXT("ExplicitlyLoaded"), bExplicitlyLoaded);
+	}
+
 	if(!PreBuildSteps.IsEmpty())
 	{
 		PreBuildSteps.Write(Writer, TEXT("PreBuildSteps"));
diff --git a/Engine/Source/Runtime/Projects/Private/PluginManager.cpp b/Engine/Source/Runtime/Projects/Private/PluginManager.cpp
--- a/Engine/Source/Runtime/Projects/Private/PluginManager.cpp
+++ b/Engine/Source/Runtime/Projects/Private/PluginManager.cpp
@@ -668,7 +668,7 @@ bool FPluginManager::ConfigureEnabledPlugins()
 		for(const TPair<FString, TSharedRef<FPlugin>>& PluginPair: AllPlugins)
 		{
 			const FPlugin& Plugin = *PluginPair.Value;
-			if (Plugin.bEnabled)
+			if (Plugin.bEnabled && !Plugin.Descriptor.bExplicitlyLoaded)
 			{
 				UE_LOG(LogPluginManager, Log, TEXT("Mounting plugin %s"), *Plugin.GetName());
 
@@ -744,6 +744,11 @@ bool FPluginManager::ConfigureEnabledPlugins()
 		IPlatformFile& PlatformFile = FPlatformFileManager::Get().GetPlatformFile();
 		for (TSharedRef<IPlugin> Plugin: GetEnabledPluginsWithContent())
 		{
+			if (Plugin->GetDescriptor().bExplicitlyLoaded)
+			{
+				continue;
+			}
+
 			if (NewPluginMountedEvent.IsBound())
 			{
 				NewPluginMountedEvent.Broadcast(*Plugin);
@@ -1189,7 +1194,7 @@ bool FPluginManager::LoadModulesForEnabledPlugins( const ELoadingPhase::Type Loa
 		const TSharedRef<FPlugin> &Plugin = PluginPair.Value;
 		SlowTask.EnterProgressFrame(1);
 
-		if ( Plugin->bEnabled )
+		if ( Plugin->bEnabled && !Plugin->Descriptor.bExplicitlyLoaded )
 		{
 			if (!TryLoadModulesForPlugin(Plugin.Get(), LoadingPhase))
 			{
@@ -1366,53 +1371,63 @@ IPluginManager::FNewPluginMountedEvent& FPluginManager::OnNewPluginMounted()
 
 void FPluginManager::MountNewlyCreatedPlugin(const FString& PluginName)
 {
-	for(TMap<FString, TSharedRef<FPlugin>>::TIterator Iter(AllPlugins); Iter; ++Iter)
+	TSharedPtr<FPlugin> Plugin = FindPluginInstance(PluginName);
+	if (Plugin.IsValid())
 	{
-		const TSharedRef<FPlugin>& Plugin = Iter.Value();
-		if (Plugin->Name == PluginName)
+		MountPluginFromExternalSource(Plugin.ToSharedRef());
+
+		// Notify any listeners that a new plugin has been mounted
+		if (NewPluginCreatedEvent.IsBound())
 		{
-			// Mark the plugin as enabled
-			Plugin->bEnabled = true;
+			NewPluginCreatedEvent.Broadcast(*Plugin);
+		}
+	}
+}
 
-			// Mount the plugin content directory
-			if (Plugin->CanContainContent() && ensure(RegisterMountPointDelegate.IsBound()))
-			{
-				FString ContentDir = Plugin->GetContentDir();
-				RegisterMountPointDelegate.Execute(Plugin->GetMountedAssetPath(), ContentDir);
+void FPluginManager::MountExplicitlyLoadedPlugin(const FString& PluginName)
+{
+	TSharedPtr<FPlugin> Plugin = FindPluginInstance(PluginName);
+	if (Plugin.IsValid() && Plugin->Descriptor.bExplicitlyLoaded)
+	{
+		MountPluginFromExternalSource(Plugin.ToSharedRef());
+	}
+}
 
-				// Register this plugin's path with the list of content directories that the editor will search
-				if (FConfigFile* EngineConfigFile = GConfig->Find(GEngineIni, false))
-				{
-					if (FConfigSection* CoreSystemSection = EngineConfigFile->Find(TEXT("Core.System")))
-					{
-						CoreSystemSection->AddUnique("Paths", MoveTemp(ContentDir));
-					}
-				}
-			}
+void FPluginManager::MountPluginFromExternalSource(const TSharedRef<FPlugin>& Plugin)
+{
+	// Mark the plugin as enabled
+	Plugin->bEnabled = true;
 
-			// If it's a code module, also load the modules for it
-			if (Plugin->Descriptor.Modules.Num() > 0)
-			{
-				// Add the plugin binaries directory
-				const FString PluginBinariesPath = FPaths::Combine(*FPaths::GetPath(Plugin->FileName), TEXT("Binaries"), FPlatformProcess::GetBinariesSubdirectory());
-				FModuleManager::Get().AddBinariesDirectory(*PluginBinariesPath, Plugin->GetLoadedFrom() == EPluginLoadedFrom::Project);
+	// Mount the plugin content directory
+	if (Plugin->CanContainContent() && ensure(RegisterMountPointDelegate.IsBound()))
+	{
+		FString ContentDir = Plugin->GetContentDir();
+		RegisterMountPointDelegate.Execute(Plugin->GetMountedAssetPath(), ContentDir);
 
-				// Load all the plugin modules
-				for (ELoadingPhase::Type LoadingPhase = (ELoadingPhase::Type)0; LoadingPhase < ELoadingPhase::Max; LoadingPhase = (ELoadingPhase::Type)(LoadingPhase + 1))
-				{
-					if (LoadingPhase != ELoadingPhase::None)
-					{
-						TryLoadModulesForPlugin(Plugin.Get(), LoadingPhase);
-					}
-				}
+		// Register this plugin's path with the list of content directories that the editor will search
+		if (FConfigFile* EngineConfigFile = GConfig->Find(GEngineIni, false))
+		{
+			if (FConfigSection* CoreSystemSection = EngineConfigFile->Find(TEXT("Core.System")))
+			{
+				CoreSystemSection->AddUnique("Paths", MoveTemp(ContentDir));
 			}
+		}
+	}
 
-			// Notify any listeners that a new plugin has been mounted
-			if (NewPluginCreatedEvent.IsBound())
+	// If it's a code module, also load the modules for it
+	if (Plugin->Descriptor.Modules.Num() > 0)
+	{
+		// Add the plugin binaries directory
+		const FString PluginBinariesPath = FPaths::Combine(*FPaths::GetPath(Plugin->FileName), TEXT("Binaries"), FPlatformProcess::GetBinariesSubdirectory());
+		FModuleManager::Get().AddBinariesDirectory(*PluginBinariesPath, Plugin->GetLoadedFrom() == EPluginLoadedFrom::Project);
+
+		// Load all the plugin modules
+		for (ELoadingPhase::Type LoadingPhase = (ELoadingPhase::Type)0; LoadingPhase < ELoadingPhase::Max; LoadingPhase = (ELoadingPhase::Type)(LoadingPhase + 1))
+		{
+			if (LoadingPhase != ELoadingPhase::None)
 			{
-				NewPluginCreatedEvent.Broadcast(*Plugin);
+				TryLoadModulesForPlugin(Plugin.Get(), LoadingPhase);
 			}
-			break;
 		}
 	}
 }
diff --git a/Engine/Source/Runtime/Projects/Private/PluginManager.h b/Engine/Source/Runtime/Projects/Private/PluginManager.h
--- a/Engine/Source/Runtime/Projects/Private/PluginManager.h
+++ b/Engine/Source/Runtime/Projects/Private/PluginManager.h
@@ -113,6 +113,7 @@ public:
 	virtual FNewPluginMountedEvent& OnNewPluginCreated() override;
 	virtual FNewPluginMountedEvent& OnNewPluginMounted() override;
 	virtual void MountNewlyCreatedPlugin(const FString& PluginName) override;
+	virtual void MountExplicitlyLoadedPlugin(const FString& PluginName) override;
 	virtual FName PackageNameFromModuleName(FName ModuleName) override;
 	virtual bool RequiresTempTargetForCodePlugin(const FProjectDescriptor* ProjectDescriptor, const FString& Platform, EBuildConfiguration Configuration, EBuildTargetType TargetType, FText& OutReason) override;
 
@@ -170,6 +171,9 @@ private:
 	/** Gets the instance of a given plugin */
 	TSharedPtr<FPlugin> FindPluginInstance(const FString& Name);
 
+	/** Mounts a plugin that was requested to be mounted from external code (either by MountNewlyCreatedPlugin or MountExplicitlyLoadedPlugin) */
+	void MountPluginFromExternalSource(const TSharedRef<FPlugin>& Plugin);
+
 private:
 	/** All of the plugins that we know about */
 	TMap< FString, TSharedRef< FPlugin > > AllPlugins;
diff --git a/Engine/Source/Runtime/Projects/Public/Interfaces/IPluginManager.h b/Engine/Source/Runtime/Projects/Public/Interfaces/IPluginManager.h
--- a/Engine/Source/Runtime/Projects/Public/Interfaces/IPluginManager.h
+++ b/Engine/Source/Runtime/Projects/Public/Interfaces/IPluginManager.h
@@ -291,6 +291,12 @@ public:
 	 */
 	virtual void MountNewlyCreatedPlugin(const FString& PluginName) = 0;
 
+	/**
+	 * Marks an explicitly loaded plugin as enabled, mounts its content and tries to load its modules.
+	 * These plugins are not loaded implicitly, but instead wait for this function to be called.
+	 */
+	virtual void MountExplicitlyLoadedPlugin(const FString& PluginName) = 0;
+
 	/**
 	* Does a reverse lookup to try to figure out what the UObject package name is for a plugin
 	*/
diff --git a/Engine/Source/Runtime/Projects/Public/PluginDescriptor.h b/Engine/Source/Runtime/Projects/Public/PluginDescriptor.h
--- a/Engine/Source/Runtime/Projects/Public/PluginDescriptor.h
+++ b/Engine/Source/Runtime/Projects/Public/PluginDescriptor.h
@@ -97,6 +97,9 @@ struct PROJECTS_API FPluginDescriptor
 	/** For auto-generated plugins that should not be listed in the plugin browser for users to disable freely. */
 	bool bIsHidden;
 
+	/** When true, this plugin's modules will not be loaded automatically nor will it's content be mounted automatically. It will load/mount when explicitly requested and LoadingPhases will be ignored */
+	bool bExplicitlyLoaded;
+
 	/** If true, this plugin from a platform extension extending another plugin */
 	bool bIsPluginExtension;
 
