PR #6013: Fix comment for PostBuildSteps (Contributed by [REDACTED_PERSON])

diff --git a/Engine/Source/Runtime/Projects/Public/PluginDescriptor.h b/Engine/Source/Runtime/Projects/Public/PluginDescriptor.h
--- a/Engine/Source/Runtime/Projects/Public/PluginDescriptor.h
+++ b/Engine/Source/Runtime/Projects/Public/PluginDescriptor.h
@@ -94,7 +94,7 @@ struct PROJECTS_API FPluginDescriptor
 	/** Pre-build steps for each host platform */
 	FCustomBuildSteps PreBuildSteps;
 
-	/** Pre-build steps for each host platform */
+	/** Post-build steps for each host platform */
 	FCustomBuildSteps PostBuildSteps;
 
 	/** Dependent plugins */
