Allow projects to opt out of having enabled by default engine plugins enabled by specify DisableEnginePluginsByDefault in the uproject. This requires all plugins to be explicitly enabled and prevents upgrades from bringing new plugins in unexpectedly.

diff --git a/Engine/Source/Developer/TargetPlatform/Private/TargetPlatformBase.cpp b/Engine/Source/Developer/TargetPlatform/Private/TargetPlatformBase.cpp
--- a/Engine/Source/Developer/TargetPlatform/Private/TargetPlatformBase.cpp
+++ b/Engine/Source/Developer/TargetPlatform/Private/TargetPlatformBase.cpp
@@ -54,7 +54,8 @@ static bool IsPluginEnabledForTarget(const IPlugin& Plugin, const FProjectDescri
 		return false;
 	}
 
-	bool bEnabledForProject = Plugin.IsEnabledByDefault();
+	const bool bAllowEnginePluginsEnabledByDefault = (Project == nullptr || !Project->bDisableEnginePluginsByDefault);
+	bool bEnabledForProject = Plugin.IsEnabledByDefault(bAllowEnginePluginsEnabledByDefault);
 	if (Project != nullptr)
 	{
 		for(const FPluginReferenceDescriptor& PluginReference : Project->Plugins)
diff --git a/Engine/Source/Programs/AutomationTool/AutomationUtils/ProjectUtils.cs b/Engine/Source/Programs/AutomationTool/AutomationUtils/ProjectUtils.cs
--- a/Engine/Source/Programs/AutomationTool/AutomationUtils/ProjectUtils.cs
+++ b/Engine/Source/Programs/AutomationTool/AutomationUtils/ProjectUtils.cs
@@ -275,21 +275,27 @@ namespace AutomationTool
 			// Keep a set of all the plugin names that have been configured. We read configuration data from different places, but only configure a plugin from the first place that it's referenced.
 			HashSet<string> ConfiguredPluginNames = new HashSet<string>(StringComparer.OrdinalIgnoreCase);
 
+			bool bAllowEnginePluginsEnabledByDefault = true;
+
 			// Find all the plugin references in the project file
-			if (ProjectDescriptor != null && ProjectDescriptor.Plugins != null)
+			if (ProjectDescriptor != null)
 			{
-				// Copy the plugin references, since we may modify the project if any plugins are missing
-				foreach (PluginReferenceDescriptor PluginReference in ProjectDescriptor.Plugins)
+				bAllowEnginePluginsEnabledByDefault = !ProjectDescriptor.DisableEnginePluginsByDefault;
+				if (ProjectDescriptor.Plugins != null)
 				{
-					if(!ConfiguredPluginNames.Contains(PluginReference.Name))
+					// Copy the plugin references, since we may modify the project if any plugins are missing
+					foreach (PluginReferenceDescriptor PluginReference in ProjectDescriptor.Plugins)
 					{
-						PluginReferenceDescriptor MissingPlugin;
-						if (!ConfigureEnabledPluginForTarget(PluginReference, ProjectDescriptor, null, Platform, Configuration, TargetType, bLoadPluginsForTargetPlatforms, AllPlugins, EnabledPlugins, out MissingPlugin))
+						if (!ConfiguredPluginNames.Contains(PluginReference.Name))
 						{
-							OutMissingPlugin = MissingPlugin;
-							return false;
+							PluginReferenceDescriptor MissingPlugin;
+							if (!ConfigureEnabledPluginForTarget(PluginReference, ProjectDescriptor, null, Platform, Configuration, TargetType, bLoadPluginsForTargetPlatforms, AllPlugins, EnabledPlugins, out MissingPlugin))
+							{
+								OutMissingPlugin = MissingPlugin;
+								return false;
+							}
+							ConfiguredPluginNames.Add(PluginReference.Name);
 						}
-						ConfiguredPluginNames.Add(PluginReference.Name);
 					}
 				}
 			}
@@ -297,7 +303,7 @@ namespace AutomationTool
 			// Add the plugins which are enabled by default
 			foreach (KeyValuePair<string, PluginInfo> PluginPair in AllPlugins)
 			{
-				if(PluginPair.Value.EnabledByDefault && !ConfiguredPluginNames.Contains(PluginPair.Key))
+				if (PluginPair.Value.IsEnabledByDefault(bAllowEnginePluginsEnabledByDefault) && !ConfiguredPluginNames.Contains(PluginPair.Key))
 				{
 					PluginReferenceDescriptor MissingPlugin;
 					if (!ConfigureEnabledPluginForTarget(new PluginReferenceDescriptor(PluginPair.Key, null, true), ProjectDescriptor, null, Platform, Configuration, TargetType, bLoadPluginsForTargetPlatforms, AllPlugins, EnabledPlugins, out MissingPlugin))
diff --git a/Engine/Source/Programs/UnrealBuildTool/Configuration/UEBuildTarget.cs b/Engine/Source/Programs/UnrealBuildTool/Configuration/UEBuildTarget.cs
--- a/Engine/Source/Programs/UnrealBuildTool/Configuration/UEBuildTarget.cs
+++ b/Engine/Source/Programs/UnrealBuildTool/Configuration/UEBuildTarget.cs
@@ -3010,22 +3010,28 @@ namespace UnrealBuildTool
 				}
 			}
 
+			bool bAllowEnginePluginsEnabledByDefault = true;
+
 			// Find a map of plugins which are explicitly referenced in the project file
-			if(ProjectDescriptor != null && ProjectDescriptor.Plugins != null)
+			if (ProjectDescriptor != null)
 			{
-				string ProjectReferenceChain = ProjectFile.GetFileName();
-				foreach(PluginReferenceDescriptor PluginReference in ProjectDescriptor.Plugins)
+				bAllowEnginePluginsEnabledByDefault = !ProjectDescriptor.DisableEnginePluginsByDefault;
+				if (ProjectDescriptor.Plugins != null)
 				{
-					if(!Rules.EnablePlugins.Contains(PluginReference.Name, StringComparer.InvariantCultureIgnoreCase) && !Rules.DisablePlugins.Contains(PluginReference.Name, StringComparer.InvariantCultureIgnoreCase))
+					string ProjectReferenceChain = ProjectFile.GetFileName();
+					foreach (PluginReferenceDescriptor PluginReference in ProjectDescriptor.Plugins)
 					{
-						// Make sure we don't have multiple references to the same plugin
-						if(!ReferencedNames.Add(PluginReference.Name))
+						if (!Rules.EnablePlugins.Contains(PluginReference.Name, StringComparer.InvariantCultureIgnoreCase) && !Rules.DisablePlugins.Contains(PluginReference.Name, StringComparer.InvariantCultureIgnoreCase))
 						{
-							Log.TraceWarning("Plugin '{0}' is listed multiple times in project file '{1}'.", PluginReference.Name, ProjectFile);
-						}
-						else
-						{
-							AddPlugin(PluginReference, ProjectReferenceChain, ExcludeFolders, NameToInstance, NameToInfo);
+							// Make sure we don't have multiple references to the same plugin
+							if (!ReferencedNames.Add(PluginReference.Name))
+							{
+								Log.TraceWarning("Plugin '{0}' is listed multiple times in project file '{1}'.", PluginReference.Name, ProjectFile);
+							}
+							else
+							{
+								AddPlugin(PluginReference, ProjectReferenceChain, ExcludeFolders, NameToInstance, NameToInfo);
+							}
 						}
 					}
 				}
@@ -3036,7 +3042,7 @@ namespace UnrealBuildTool
 			{
 				foreach(PluginInfo Plugin in NameToInfo.Values)
 				{
-					if(Plugin.EnabledByDefault && !ReferencedNames.Contains(Plugin.Name))
+					if(Plugin.IsEnabledByDefault(bAllowEnginePluginsEnabledByDefault) && !ReferencedNames.Contains(Plugin.Name))
 					{
 						ReferencedNames.Add(Plugin.Name);
 
diff --git a/Engine/Source/Programs/UnrealBuildTool/System/Plugins.cs b/Engine/Source/Programs/UnrealBuildTool/System/Plugins.cs
--- a/Engine/Source/Programs/UnrealBuildTool/System/Plugins.cs
+++ b/Engine/Source/Programs/UnrealBuildTool/System/Plugins.cs
@@ -109,19 +109,23 @@ namespace UnrealBuildTool
 		/// <summary>
 		/// Determines whether the plugin should be enabled by default
 		/// </summary>
-		public bool EnabledByDefault
+		public bool IsEnabledByDefault(bool bAllowEnginePluginsEnabledByDefault)
 		{
-			get
+			if (Descriptor.bEnabledByDefault.HasValue)
 			{
-				if(Descriptor.bEnabledByDefault.HasValue)
+				if (Descriptor.bEnabledByDefault.Value)
 				{
-					return Descriptor.bEnabledByDefault.Value;
+					return (LoadedFrom == PluginLoadedFrom.Project ? true : bAllowEnginePluginsEnabledByDefault);
 				}
 				else
 				{
-					return (LoadedFrom == PluginLoadedFrom.Project);
+					return false;
 				}
 			}
+			else
+			{
+				return (LoadedFrom == PluginLoadedFrom.Project);
+			}
 		}
 
 		/// <summary>
@@ -550,7 +554,8 @@ namespace UnrealBuildTool
 				return false;
 			}
 
-			bool bEnabled = Plugin.EnabledByDefault;
+			bool bAllowEnginePluginsEnabledByDefault = (Project == null ? true : !Project.DisableEnginePluginsByDefault);
+			bool bEnabled = Plugin.IsEnabledByDefault(bAllowEnginePluginsEnabledByDefault);
 			if (Project != null && Project.Plugins != null)
 			{
 				foreach (PluginReferenceDescriptor PluginReference in Project.Plugins)
diff --git a/Engine/Source/Programs/UnrealBuildTool/System/ProjectDescriptor.cs b/Engine/Source/Programs/UnrealBuildTool/System/ProjectDescriptor.cs
--- a/Engine/Source/Programs/UnrealBuildTool/System/ProjectDescriptor.cs
+++ b/Engine/Source/Programs/UnrealBuildTool/System/ProjectDescriptor.cs
@@ -110,6 +110,11 @@ namespace UnrealBuildTool
 		/// </summary>
 		public bool IsEnterpriseProject;
 
+		/// <summary>
+		/// Indicates that enabled by default engine plugins should not be enabled unless explicitly enabled by the project or target files.
+		/// </summary>
+		public bool DisableEnginePluginsByDefault;
+
 		/// <summary>
 		/// Constructor.
 		/// </summary>
@@ -117,6 +122,7 @@ namespace UnrealBuildTool
 		{
 			FileVersion = (int)ProjectDescriptorVersion.Latest;
 			IsEnterpriseProject = false;
+			DisableEnginePluginsByDefault = false;
 		}
 
 		/// <summary>
@@ -145,6 +151,7 @@ namespace UnrealBuildTool
 			RawObject.TryGetStringField("Category", out Category);
 			RawObject.TryGetStringField("Description", out Description);
 			RawObject.TryGetBoolField("Enterprise", out IsEnterpriseProject);
+			RawObject.TryGetBoolField("DisableEnginePluginsByDefault", out DisableEnginePluginsByDefault);
 
 			// Read the modules
 			JsonObject[] ModulesArray;
diff --git a/Engine/Source/Runtime/Projects/Private/PluginManager.cpp b/Engine/Source/Runtime/Projects/Private/PluginManager.cpp
--- a/Engine/Source/Runtime/Projects/Private/PluginManager.cpp
+++ b/Engine/Source/Runtime/Projects/Private/PluginManager.cpp
@@ -105,11 +105,11 @@ FString FPlugin::GetMountedAssetPath() const
 	return Path;
 }
 
-bool FPlugin::IsEnabledByDefault() const
+bool FPlugin::IsEnabledByDefault(const bool bAllowEnginePluginsEnabledByDefault) const
 {
 	if (Descriptor.EnabledByDefault == EPluginEnabledByDefault::Enabled)
 	{
-		return true;
+		return (GetLoadedFrom() == EPluginLoadedFrom::Project ? true : bAllowEnginePluginsEnabledByDefault);
 	}
 	else if (Descriptor.EnabledByDefault == EPluginEnabledByDefault::Disabled)
 	{
@@ -587,10 +587,13 @@ bool FPluginManager::ConfigureEnabledPlugins()
 				}
 			}
 
+			bool bAllowEnginePluginsEnabledByDefault = true;
 			// Find all the plugin references in the project file
 			const FProjectDescriptor* ProjectDescriptor = IProjectManager::Get().GetCurrentProject();
 			if (ProjectDescriptor != nullptr)
 			{
+				bAllowEnginePluginsEnabledByDefault = !ProjectDescriptor->bDisableEnginePluginsByDefault;
+
 				// Copy the plugin references, since we may modify the project if any plugins are missing
 				TArray<FPluginReferenceDescriptor> PluginReferences(ProjectDescriptor->Plugins);
 				for (const FPluginReferenceDescriptor& PluginReference : PluginReferences)
@@ -609,7 +612,7 @@ bool FPluginManager::ConfigureEnabledPlugins()
 			// Add the plugins which are enabled by default
 			for (const TPair<FString, TSharedRef<FPlugin>>& PluginPair : AllPlugins)
 			{
-				if(PluginPair.Value->IsEnabledByDefault() && !ConfiguredPluginNames.Contains(PluginPair.Key))
+				if (PluginPair.Value->IsEnabledByDefault(bAllowEnginePluginsEnabledByDefault) && !ConfiguredPluginNames.Contains(PluginPair.Key))
 				{
 					if (!ConfigureEnabledPluginForCurrentTarget(FPluginReferenceDescriptor(PluginPair.Key, true), EnabledPlugins))
 					{
@@ -837,8 +840,11 @@ bool FPluginManager::GetCodePluginsForProject(const FProjectDescriptor* ProjectD
 	TSet<FString> ConfiguredPluginNames;
 
 	// Find all the plugin references in the project file
+	bool bAllowEnginePluginsEnabledByDefault = true;
 	if (ProjectDescriptor != nullptr)
 	{
+		bAllowEnginePluginsEnabledByDefault = !ProjectDescriptor->bDisableEnginePluginsByDefault;
+
 		// Copy the plugin references, since we may modify the project if any plugins are missing
 		TArray<FPluginReferenceDescriptor> PluginReferences(ProjectDescriptor->Plugins);
 		for (const FPluginReferenceDescriptor& PluginReference : PluginReferences)
@@ -857,7 +863,7 @@ bool FPluginManager::GetCodePluginsForProject(const FProjectDescriptor* ProjectD
 	// Add the plugins which are enabled by default
 	for (const TPair<FString, TSharedRef<FPlugin>>& PluginPair : AllPlugins)
 	{
-		if(PluginPair.Value->IsEnabledByDefault() && !ConfiguredPluginNames.Contains(PluginPair.Key))
+		if (PluginPair.Value->IsEnabledByDefault(bAllowEnginePluginsEnabledByDefault) && !ConfiguredPluginNames.Contains(PluginPair.Key))
 		{
 			if (!ConfigureEnabledPluginForTarget(FPluginReferenceDescriptor(PluginPair.Key, true), ProjectDescriptor, FString(), Platform, Configuration, TargetType, bLoadPluginsForTargetPlatforms, AllPlugins, EnabledPlugins, OutMissingPlugin))
 			{
diff --git a/Engine/Source/Runtime/Projects/Private/PluginManager.h b/Engine/Source/Runtime/Projects/Private/PluginManager.h
--- a/Engine/Source/Runtime/Projects/Private/PluginManager.h
+++ b/Engine/Source/Runtime/Projects/Private/PluginManager.h
@@ -59,7 +59,7 @@ public:
 		return bEnabled;
 	}
 
-	virtual bool IsEnabledByDefault() const override;
+	virtual bool IsEnabledByDefault(bool bAllowEnginePluginsEnabledByDefault) const override;
 
 	virtual bool IsHidden() const override
 	{
diff --git a/Engine/Source/Runtime/Projects/Private/ProjectDescriptor.cpp b/Engine/Source/Runtime/Projects/Private/ProjectDescriptor.cpp
--- a/Engine/Source/Runtime/Projects/Private/ProjectDescriptor.cpp
+++ b/Engine/Source/Runtime/Projects/Private/ProjectDescriptor.cpp
@@ -14,6 +14,7 @@ FProjectDescriptor::FProjectDescriptor()
 	FileVersion = EProjectDescriptorVersion::Latest;
 	EpicSampleNameHash = 0;
 	bIsEnterpriseProject = false;
+	bDisableEnginePluginsByDefault = false;
 }
 
 void FProjectDescriptor::Sign(const FString& FilePath)
@@ -110,6 +111,7 @@ bool FProjectDescriptor::Read(const FJsonObject& Object, const FString& PathToPr
 	Object.TryGetStringField(TEXT("Category"), Category);
 	Object.TryGetStringField(TEXT("Description"), Description);
 	Object.TryGetBoolField(TEXT("Enterprise"), bIsEnterpriseProject);
+	Object.TryGetBoolField(TEXT("DisableEnginePluginsByDefault"), bDisableEnginePluginsByDefault);
 
 	// Read the modules
 	if(!FModuleDescriptor::ReadArray(Object, TEXT("Modules"), Modules, OutFailReason))
diff --git a/Engine/Source/Runtime/Projects/Private/ProjectManager.cpp b/Engine/Source/Runtime/Projects/Private/ProjectManager.cpp
--- a/Engine/Source/Runtime/Projects/Private/ProjectManager.cpp
+++ b/Engine/Source/Runtime/Projects/Private/ProjectManager.cpp
@@ -252,11 +252,13 @@ void FProjectManager::ClearSupportedTargetPlatformsForCurrentProject()
 PROJECTS_API bool HasDefaultPluginSettings(const FString& Platform)
 {
 	const FProjectDescriptor* CurrentProject = IProjectManager::Get().GetCurrentProject();
+	bool bAllowEnginePluginsEnabledByDefault = true;
 
 	// Get settings for the plugins which are currently enabled or disabled by the project file
 	TMap<FString, bool> ConfiguredPlugins;
 	if (CurrentProject != nullptr)
 	{
+		bAllowEnginePluginsEnabledByDefault = !CurrentProject->bDisableEnginePluginsByDefault;
 		for (const FPluginReferenceDescriptor& PluginReference : CurrentProject->Plugins)
 		{
 			ConfiguredPlugins.Add(PluginReference.Name, PluginReference.IsEnabledForPlatform(Platform));
@@ -268,7 +270,8 @@ PROJECTS_API bool HasDefaultPluginSettings(const FString& Platform)
 	{
 		if (Plugin->GetDescriptor().SupportsTargetPlatform(Platform))
 		{
-			bool bEnabled = Plugin->IsEnabledByDefault();
+			const bool bEnabledByDefault = Plugin->IsEnabledByDefault(bAllowEnginePluginsEnabledByDefault);
+			bool bEnabled = bEnabledByDefault;
 
 			bool* EnabledPtr = ConfiguredPlugins.Find(Plugin->GetName());
 			if(EnabledPtr != nullptr)
@@ -276,7 +279,7 @@ PROJECTS_API bool HasDefaultPluginSettings(const FString& Platform)
 				bEnabled = *EnabledPtr;
 			}
 
-			bool bEnabledInDefaultExe = (Plugin->GetLoadedFrom() == EPluginLoadedFrom::Engine && Plugin->IsEnabledByDefault() && !Plugin->GetDescriptor().bInstalled);
+			bool bEnabledInDefaultExe = (Plugin->GetLoadedFrom() == EPluginLoadedFrom::Engine && bEnabledByDefault && !Plugin->GetDescriptor().bInstalled);
 			if(bEnabled != bEnabledInDefaultExe)
 			{
 				return false;
@@ -346,7 +349,7 @@ bool FProjectManager::SetPluginEnabled(const FString& PluginName, bool bEnabled,
 		{
 			// Get the default list of enabled plugins
 			TArray<FString> DefaultEnabledPlugins;
-			GetDefaultEnabledPlugins(DefaultEnabledPlugins, false);
+			GetDefaultEnabledPlugins(DefaultEnabledPlugins, false, !CurrentProject->bDisableEnginePluginsByDefault);
 
 			// Check the enabled state is the same in that
 			if(DefaultEnabledPlugins.Contains(PluginName) == bEnabled)
@@ -385,12 +388,12 @@ bool FProjectManager::RemovePluginReference(const FString& PluginName, FText& Ou
 	return bPluginFound;
 }
 
-void FProjectManager::GetDefaultEnabledPlugins(TArray<FString>& OutPluginNames, bool bIncludeInstalledPlugins)
+void FProjectManager::GetDefaultEnabledPlugins(TArray<FString>& OutPluginNames, bool bIncludeInstalledPlugins, bool bAllowEnginePluginsEnabledByDefault)
 {
 	// Add all the game plugins and everything marked as enabled by default
 	for (TSharedRef<IPlugin>& Plugin : IPluginManager::Get().GetDiscoveredPlugins())
 	{
-		if(Plugin->IsEnabledByDefault())
+		if(Plugin->IsEnabledByDefault(bAllowEnginePluginsEnabledByDefault))
 		{
 			if(bIncludeInstalledPlugins || !Plugin->GetDescriptor().bInstalled)
 			{
diff --git a/Engine/Source/Runtime/Projects/Private/ProjectManager.h b/Engine/Source/Runtime/Projects/Private/ProjectManager.h
--- a/Engine/Source/Runtime/Projects/Private/ProjectManager.h
+++ b/Engine/Source/Runtime/Projects/Private/ProjectManager.h
@@ -45,7 +45,7 @@ private:
 	static void QueryStatusForProjectImpl(const FProjectDescriptor& Project, const FString& FilePath, FProjectStatus& OutProjectStatus);
 
 	/** Gets the list of plugins enabled by default, excluding the project overrides */
-	static void GetDefaultEnabledPlugins(TArray<FString>& PluginNames, bool bIncludeInstalledPlugins);
+	static void GetDefaultEnabledPlugins(TArray<FString>& PluginNames, bool bIncludeInstalledPlugins, bool bAllowEnginePluginsEnabledByDefault);
 
 	/** The project that is currently loaded in the editor */
 	TSharedPtr< FProjectDescriptor > CurrentProject;
diff --git a/Engine/Source/Runtime/Projects/Public/Interfaces/IPluginManager.h b/Engine/Source/Runtime/Projects/Public/Interfaces/IPluginManager.h
--- a/Engine/Source/Runtime/Projects/Public/Interfaces/IPluginManager.h
+++ b/Engine/Source/Runtime/Projects/Public/Interfaces/IPluginManager.h
@@ -126,7 +126,7 @@ public:
 	 *
 	 * @return True if the plugin is currently enabled by default.
 	 */
-	virtual bool IsEnabledByDefault() const = 0;
+	virtual bool IsEnabledByDefault(bool bAllowEnginePluginsEnabledByDefault) const = 0;
 
 	/**
 	 * Determines if the plugin is should be displayed in-editor for the user to enable/disable freely.
diff --git a/Engine/Source/Runtime/Projects/Public/ProjectDescriptor.h b/Engine/Source/Runtime/Projects/Public/ProjectDescriptor.h
--- a/Engine/Source/Runtime/Projects/Public/ProjectDescriptor.h
+++ b/Engine/Source/Runtime/Projects/Public/ProjectDescriptor.h
@@ -94,6 +94,9 @@ struct PROJECTS_API FProjectDescriptor
 	/** Indicates if this project is an Enterprise project */
 	bool bIsEnterpriseProject;
 
+	/** Indicates that enabled by default engine plugins should not be enabled unless explicitly enabled by the project or target files. */
+	bool bDisableEnginePluginsByDefault;
+
 	/** Constructor. */
 	FProjectDescriptor();
 
