#UE4 Added missing case for ClientOnlyNoCommandlet in IsCompiledInCurrentConfiguration

diff --git a/Engine/Source/Runtime/Projects/Private/ModuleDescriptor.cpp b/Engine/Source/Runtime/Projects/Private/ModuleDescriptor.cpp
--- a/Engine/Source/Runtime/Projects/Private/ModuleDescriptor.cpp
+++ b/Engine/Source/Runtime/Projects/Private/ModuleDescriptor.cpp
@@ -411,6 +411,7 @@ bool FModuleDescriptor::IsCompiledInConfiguration(const FString& Platform, EBuil
     case EHostType::ServerOnly:
         return TargetType != EBuildTargetType::Program && TargetType != EBuildTargetType::Client;
     case EHostType::ClientOnly:
+	case EHostType::ClientOnlyNoCommandlet:
         return TargetType != EBuildTargetType::Program && TargetType != EBuildTargetType::Server;
     }
 
