#UE4 Added AddToPluginsList to IPluginManager. This will add a single plugin to the manager, which is more efficient than refreshing the entire list.

diff --git a/Engine/Source/Runtime/Projects/Private/PluginManager.cpp b/Engine/Source/Runtime/Projects/Private/PluginManager.cpp
--- a/Engine/Source/Runtime/Projects/Private/PluginManager.cpp
+++ b/Engine/Source/Runtime/Projects/Private/PluginManager.cpp
@@ -200,6 +200,65 @@ void FPluginManager::RefreshPluginsList()
 	}
 }
 
+bool FPluginManager::AddToPluginsList(const FString& PluginFilename)
+{
+#if (WITH_ENGINE && !IS_PROGRAM) || WITH_PLUGIN_SUPPORT
+	// No need to readd if it already exists
+	FString PluginName = FPaths::GetBaseFilename(PluginFilename);
+	if (AllPlugins.Contains(PluginName))
+	{
+		return true;
+	}
+
+	// Read the plugin and load it
+	FPluginDescriptor Descriptor;
+	FText FailureReason;
+	if (Descriptor.Load(PluginFilename, FailureReason))
+	{
+		// Determine the plugin type
+		EPluginType PluginType = EPluginType::External;
+		if (PluginFilename.StartsWith(FPaths::EngineDir()) || PluginFilename.StartsWith(FPaths::EnginePlatformExtensionsDir()))
+		{
+			PluginType = EPluginType::Engine;
+		}
+		else if (PluginFilename.StartsWith(FPaths::EnterpriseDir()))
+		{
+			PluginType = EPluginType::Enterprise;
+		}
+		else if (PluginFilename.StartsWith(FPaths::ProjectModsDir()))
+		{
+			PluginType = EPluginType::Mod;
+		}
+		else if (PluginFilename.StartsWith(FPaths::GetPath(FPaths::GetProjectFilePath())))
+		{
+			PluginType = EPluginType::Project;
+		}
+
+		// Create the plugin
+		TMap<FString, TSharedRef<FPlugin>> NewPlugins;
+		TArray<TSharedRef<FPlugin>> ChildPlugins;
+		CreatePluginObject(PluginFilename, Descriptor, PluginType, NewPlugins, ChildPlugins);
+		ensureMsgf(ChildPlugins.Num() == 0, TEXT("AddToPluginsList does not allow plugins with bIsPluginExtension set to true. Plugin: %s"), *PluginFilename);
+		ensure(NewPlugins.Num() == 1);
+		
+		// Add the loaded plugin
+		TSharedRef<FPlugin>* NewPlugin = NewPlugins.Find(PluginName);
+		if (ensure(NewPlugin))
+		{
+			AllPlugins.Add(PluginName, *NewPlugin);
+		}
+
+		return true;
+	}
+	else
+	{
+		UE_LOG(LogPluginManager, Warning, TEXT("AddToPluginsList failed to load plugin %s. Reason: %s"), *PluginFilename, *FailureReason.ToString());
+	}
+#endif
+
+	return false;
+}
+
 void FPluginManager::DiscoverAllPlugins()
 {
 	ensure( AllPlugins.Num() == 0 );		// Should not have already been initialized!
diff --git a/Engine/Source/Runtime/Projects/Private/PluginManager.h b/Engine/Source/Runtime/Projects/Private/PluginManager.h
--- a/Engine/Source/Runtime/Projects/Private/PluginManager.h
+++ b/Engine/Source/Runtime/Projects/Private/PluginManager.h
@@ -96,6 +96,7 @@ public:
 
 	/** IPluginManager interface */
 	virtual void RefreshPluginsList() override;
+	virtual bool AddToPluginsList( const FString& PluginFilename ) override;
 	virtual bool LoadModulesForEnabledPlugins( const ELoadingPhase::Type LoadingPhase ) override;
 	virtual void GetLocalizationPathsForEnabledPlugins( TArray<FString>& OutLocResPaths ) override;
 	virtual void SetRegisterMountPointDelegate( const FRegisterMountPointDelegate& Delegate ) override;
diff --git a/Engine/Source/Runtime/Projects/Public/Interfaces/IPluginManager.h b/Engine/Source/Runtime/Projects/Public/Interfaces/IPluginManager.h
--- a/Engine/Source/Runtime/Projects/Public/Interfaces/IPluginManager.h
+++ b/Engine/Source/Runtime/Projects/Public/Interfaces/IPluginManager.h
@@ -179,6 +179,13 @@ public:
 	 */
 	virtual void RefreshPluginsList() = 0;
 
+	/**
+	 * Adds a single plugin to the list of plugins. Faster than refreshing all plugins with RefreshPluginsList() when you only want to add one. Does nothing if already in the list.
+	 * 
+	 * @return True if the plugin was added or already in the list. False if it failed to load.
+	 */
+	virtual bool AddToPluginsList( const FString& PluginFilename ) = 0;
+
 	/**
 	 * Loads all plug-ins
 	 *
