#UE4 Write the target-enabled/disabled plugins lists to the target receipt instead of compiling it into the executable in non-monolitic editor builds. This allows them to be used in shared build environments

diff --git a/Engine/Source/Developer/DesktopPlatform/Private/TargetReceipt.cpp b/Engine/Source/Developer/DesktopPlatform/Private/TargetReceipt.cpp
--- a/Engine/Source/Developer/DesktopPlatform/Private/TargetReceipt.cpp
+++ b/Engine/Source/Developer/DesktopPlatform/Private/TargetReceipt.cpp
@@ -114,6 +114,29 @@ bool FTargetReceipt::Read(const FString& FileName)
 		}
 	}
 
+	// Read the list of rules-enabled and disabled plugins
+	const TArray<TSharedPtr<FJsonValue>>* PluginsArray;
+	if (Object->TryGetArrayField(TEXT("Plugins"), PluginsArray))
+	{
+		for (const TSharedPtr<FJsonValue>& PluginValue : *PluginsArray)
+		{
+			const TSharedPtr<FJsonObject>* PluginObject;
+			if (!PluginValue->TryGetObject(PluginObject))
+			{
+				return false;
+			}
+
+			FString PluginName;
+			bool bPluginEnabled;
+			if (!(*PluginObject)->TryGetStringField(TEXT("Name"), PluginName) || !(*PluginObject)->TryGetBoolField(TEXT("Enabled"), bPluginEnabled))
+			{
+				return false;
+			}
+
+			PluginNameToEnabledState.Add(PluginName, bPluginEnabled);
+		}
+	}
+
 	// Read the list of additional properties
 	const TArray<TSharedPtr<FJsonValue>>* AdditionalPropertiesArray;
 	if (Object->TryGetArrayField(TEXT("AdditionalProperties"), AdditionalPropertiesArray))
diff --git a/Engine/Source/Developer/DesktopPlatform/Public/TargetReceipt.h b/Engine/Source/Developer/DesktopPlatform/Public/TargetReceipt.h
--- a/Engine/Source/Developer/DesktopPlatform/Public/TargetReceipt.h
+++ b/Engine/Source/Developer/DesktopPlatform/Public/TargetReceipt.h
@@ -76,6 +76,9 @@ struct DESKTOPPLATFORM_API FTargetReceipt
 	/** All the runtime dependencies that this target relies on */
 	TArray<FRuntimeDependency> RuntimeDependencies;
 
+	/** All plugins that were either enabled or disabled via the target rules. */
+	TMap<FString, bool> PluginNameToEnabledState;
+
 	/** Additional build properties passed through from the module rules */
 	TArray<FReceiptProperty> AdditionalProperties;
 
diff --git a/Engine/Source/Programs/UnrealBuildTool/Configuration/UEBuildTarget.cs b/Engine/Source/Programs/UnrealBuildTool/Configuration/UEBuildTarget.cs
--- a/Engine/Source/Programs/UnrealBuildTool/Configuration/UEBuildTarget.cs
+++ b/Engine/Source/Programs/UnrealBuildTool/Configuration/UEBuildTarget.cs
@@ -830,7 +830,8 @@ namespace UnrealBuildTool
 			}
 
 			// Make sure that we don't explicitly enable or disable any plugins through the target rules. We can't do this with the shared build environment because it requires recompiling the "Projects" engine module.
-			if(ThisRules.EnablePlugins.Count > 0 || ThisRules.DisablePlugins.Count > 0)
+			bool bUsesTargetReceiptToEnablePlugins = (ThisRules.Type == TargetType.Editor && ThisRules.LinkType != TargetLinkType.Monolithic);
+			if (!bUsesTargetReceiptToEnablePlugins && (ThisRules.EnablePlugins.Count > 0 || ThisRules.DisablePlugins.Count > 0))
 			{
 				throw new BuildException("Explicitly enabling and disabling plugins for a target is only supported when using a unique build environment (eg. for monolithic game targets).");
 			}
@@ -1461,6 +1462,16 @@ namespace UnrealBuildTool
 				}
 			}
 
+			// Add Rules-enabled and disabled plugins
+			foreach (string EnabledPluginName in Rules.EnablePlugins)
+			{
+				Receipt.PluginNameToEnabledState[EnabledPluginName] = true;
+			}
+			foreach (string DisabledPluginName in Rules.DisablePlugins)
+			{
+				Receipt.PluginNameToEnabledState[DisabledPluginName] = false;
+			}
+
 			// Find all the modules which are part of this target
 			HashSet<UEBuildModule> UniqueLinkedModules = new HashSet<UEBuildModule>();
 			foreach (UEBuildBinary Binary in Binaries)
diff --git a/Engine/Source/Programs/UnrealBuildTool/System/TargetReceipt.cs b/Engine/Source/Programs/UnrealBuildTool/System/TargetReceipt.cs
--- a/Engine/Source/Programs/UnrealBuildTool/System/TargetReceipt.cs
+++ b/Engine/Source/Programs/UnrealBuildTool/System/TargetReceipt.cs
@@ -310,6 +310,11 @@ namespace UnrealBuildTool
 		/// </summary>
 		public RuntimeDependencyList RuntimeDependencies = new RuntimeDependencyList();
 
+		/// <summary>
+		/// All plugins that were either enabled or disabled via the target rules.
+		/// </summary>
+		public Dictionary<string, bool> PluginNameToEnabledState = new Dictionary<string, bool>();
+
 		/// <summary>
 		/// Additional build properties passed through from the module rules
 		/// </summary>
@@ -359,6 +364,13 @@ namespace UnrealBuildTool
 			{
 				RuntimeDependencies.Add(new RuntimeDependency(OtherRuntimeDependency));
 			}
+			foreach (KeyValuePair<string, bool> Pair in Other.PluginNameToEnabledState)
+			{
+				if (!PluginNameToEnabledState.ContainsKey(Pair.Key))
+				{
+					PluginNameToEnabledState.Add(Pair.Key, Pair.Value);
+				}
+			}
 			AdditionalProperties.AddRange(Other.AdditionalProperties);
 		}
 
@@ -392,6 +404,13 @@ namespace UnrealBuildTool
 					RuntimeDependencies.Add(OtherRuntimeDependency);
 				}
 			}
+			foreach (KeyValuePair<string, bool> Pair in Other.PluginNameToEnabledState)
+			{
+				if (!PluginNameToEnabledState.ContainsKey(Pair.Key))
+				{
+					PluginNameToEnabledState.Add(Pair.Key, Pair.Value);
+				}
+			}
 		}
 
 		/// <summary>
@@ -618,6 +637,24 @@ namespace UnrealBuildTool
 				}
 			}
 
+			// Read the enabled/disabled plugins
+			JsonObject[] PluginObjects;
+			if (RawObject.TryGetObjectArrayField("Plugins", out PluginObjects))
+			{
+				foreach (JsonObject PluginObject in PluginObjects)
+				{
+					string PluginName;
+					if (PluginObject.TryGetStringField("Name", out PluginName))
+					{
+						bool PluginEnabled;
+						if (PluginObject.TryGetBoolField("Enabled", out PluginEnabled))
+						{
+							Receipt.PluginNameToEnabledState.Add(PluginName, PluginEnabled);
+						}
+					}
+				}
+			}
+
 			// Read the additional properties
 			JsonObject[] AdditionalPropertyObjects;
 			if(RawObject.TryGetObjectArrayField("AdditionalProperties", out AdditionalPropertyObjects))
@@ -736,7 +773,20 @@ namespace UnrealBuildTool
 				}
 				Writer.WriteArrayEnd();
 
-				if(AdditionalProperties.Count > 0)
+				if (PluginNameToEnabledState.Count > 0)
+				{
+					Writer.WriteArrayStart("Plugins");
+					foreach (KeyValuePair<string, bool> PluginNameToEnabledStatePair in PluginNameToEnabledState.OrderBy(x => x.Key))
+					{
+						Writer.WriteObjectStart();
+						Writer.WriteValue("Name", PluginNameToEnabledStatePair.Key);
+						Writer.WriteValue("Enabled", PluginNameToEnabledStatePair.Value);
+						Writer.WriteObjectEnd();
+					}
+					Writer.WriteArrayEnd();
+				}
+
+				if (AdditionalProperties.Count > 0)
 				{
 					Writer.WriteArrayStart("AdditionalProperties");
 					foreach (ReceiptProperty AdditionalProperty in AdditionalProperties)
diff --git a/Engine/Source/Runtime/Projects/Private/PluginManager.cpp b/Engine/Source/Runtime/Projects/Private/PluginManager.cpp
--- a/Engine/Source/Runtime/Projects/Private/PluginManager.cpp
+++ b/Engine/Source/Runtime/Projects/Private/PluginManager.cpp
@@ -20,6 +20,11 @@
 #include "ProjectManager.h"
 #include "PluginManifest.h"
 #include "HAL/PlatformTime.h"
+#if READ_TARGET_ENABLED_PLUGINS_FROM_RECEIPT
+#include "IDesktopPlatform.h"
+#include "DesktopPlatformModule.h"
+#include "TargetReceipt.h"
+#endif
 
 DEFINE_LOG_CATEGORY_STATIC( LogPluginManager, Log, All );
 
@@ -618,7 +623,44 @@ bool FPluginManager::ConfigureEnabledPlugins()
 
 		if (!FParse::Param(FCommandLine::Get(), TEXT("NoEnginePlugins")))
 		{
-			// Configure the plugins that were enabled from the target file
+#if READ_TARGET_ENABLED_PLUGINS_FROM_RECEIPT
+			// Configure the plugins that were enabled or disabled from the target file using the target receipt file
+			FString EditorTargetFileName;
+			for (const FTargetInfo& Target : FDesktopPlatformModule::Get()->GetTargetsForCurrentProject())
+			{
+				if (Target.Type == FApp::GetBuildTargetType())
+				{
+					if (FPaths::IsUnderDirectory(Target.Path, FPlatformMisc::ProjectDir()))
+					{
+						EditorTargetFileName = FTargetReceipt::GetDefaultPath(FPlatformMisc::ProjectDir(), *Target.Name, FPlatformProcess::GetBinariesSubdirectory(), FApp::GetBuildConfiguration(), nullptr);
+					}
+					else if (FPaths::IsUnderDirectory(Target.Path, FPaths::EngineDir()))
+					{
+						EditorTargetFileName = FTargetReceipt::GetDefaultPath(*FPaths::EngineDir(), *Target.Name, FPlatformProcess::GetBinariesSubdirectory(), FApp::GetBuildConfiguration(), nullptr);
+					}
+					break;
+				}
+			}
+
+			FTargetReceipt Receipt;
+			if (Receipt.Read(EditorTargetFileName))
+			{
+				for (const TPair<FString, bool>& Pair : Receipt.PluginNameToEnabledState)
+				{
+					const FString& PluginName = Pair.Key;
+					const bool bEnabled = Pair.Value;
+					if (!ConfiguredPluginNames.Contains(PluginName))
+					{
+						if (!ConfigureEnabledPluginForCurrentTarget(FPluginReferenceDescriptor(PluginName, bEnabled), EnabledPlugins))
+						{
+							return false;
+						}
+						ConfiguredPluginNames.Add(PluginName);
+					}
+				}
+			}
+#else
+			// Configure the plugins that were enabled from the target file using defines
 			TArray<FString> TargetEnabledPlugins = { UBT_TARGET_ENABLED_PLUGINS };
 			for (const FString& TargetEnabledPlugin : TargetEnabledPlugins)
 			{
@@ -632,7 +674,7 @@ bool FPluginManager::ConfigureEnabledPlugins()
 				}
 			}
 
-			// Configure the plugins that were disabled from the target file
+			// Configure the plugins that were disabled from the target file using defines
 			TArray<FString> TargetDisabledPlugins = { UBT_TARGET_DISABLED_PLUGINS };
 			for (const FString& TargetDisabledPlugin : TargetDisabledPlugins)
 			{
@@ -645,6 +687,7 @@ bool FPluginManager::ConfigureEnabledPlugins()
 					ConfiguredPluginNames.Add(TargetDisabledPlugin);
 				}
 			}
+#endif // READ_TARGET_ENABLED_PLUGINS_FROM_RECEIPT
 
 			bool bAllowEnginePluginsEnabledByDefault = true;
 			// Find all the plugin references in the project file
diff --git a/Engine/Source/Runtime/Projects/Projects.Build.cs b/Engine/Source/Runtime/Projects/Projects.Build.cs
--- a/Engine/Source/Runtime/Projects/Projects.Build.cs
+++ b/Engine/Source/Runtime/Projects/Projects.Build.cs
@@ -24,19 +24,37 @@ namespace UnrealBuildTool.Rules
 				}
 			);
 
-			List<string> EnabledPluginStrings = new List<string>();
-			foreach(string EnablePlugin in Target.EnablePlugins)
+			// Monolithic and non-editor builds compile in the enabled/disabled plugins. Non-monolithic editor builds save them in the plugin receipt to avoid invalidating the shared build environment.
+			if (Target.Type == TargetType.Editor && Target.LinkType != TargetLinkType.Monolithic)
 			{
-				EnabledPluginStrings.Add(String.Format("TEXT(\"{0}\")", EnablePlugin));
-			}
-			PrivateDefinitions.Add(String.Format("UBT_TARGET_ENABLED_PLUGINS={0}", String.Join(", ", EnabledPluginStrings)));
+				PublicDefinitions.Add("READ_TARGET_ENABLED_PLUGINS_FROM_RECEIPT=1");
 
-			List<string> DisabledPluginStrings = new List<string>();
-			foreach(string DisablePlugin in Target.DisablePlugins)
+				// Include DesktopPlatform to work with the target receipt which contains the enabled/disabled plugins
+				PrivateDependencyModuleNames.AddRange(
+					new string[]
+					{
+						"DesktopPlatform"
+					}
+				);
+			}
+			else
 			{
-				DisabledPluginStrings.Add(String.Format("TEXT(\"{0}\")", DisablePlugin));
+				PublicDefinitions.Add("READ_TARGET_ENABLED_PLUGINS_FROM_RECEIPT=0");
+
+				List<string> EnabledPluginStrings = new List<string>();
+				foreach(string EnablePlugin in Target.EnablePlugins)
+				{
+					EnabledPluginStrings.Add(String.Format("TEXT(\"{0}\")", EnablePlugin));
+				}
+				PrivateDefinitions.Add(String.Format("UBT_TARGET_ENABLED_PLUGINS={0}", String.Join(", ", EnabledPluginStrings)));
+
+				List<string> DisabledPluginStrings = new List<string>();
+				foreach(string DisablePlugin in Target.DisablePlugins)
+				{
+					DisabledPluginStrings.Add(String.Format("TEXT(\"{0}\")", DisablePlugin));
+				}
+				PrivateDefinitions.Add(String.Format("UBT_TARGET_DISABLED_PLUGINS={0}", String.Join(", ", DisabledPluginStrings)));
 			}
-			PrivateDefinitions.Add(String.Format("UBT_TARGET_DISABLED_PLUGINS={0}", String.Join(", ", DisabledPluginStrings)));
 
 			if (Target.bIncludePluginsForTargetPlatforms)
 			{
