#UE4 Add fallback if the target info doesnt contain known paths in the plugin manager. The target file is assumed to be either in the project or engine folder, preferring project

diff --git a/Engine/Source/Runtime/Projects/Private/PluginManager.cpp b/Engine/Source/Runtime/Projects/Private/PluginManager.cpp
--- a/Engine/Source/Runtime/Projects/Private/PluginManager.cpp
+++ b/Engine/Source/Runtime/Projects/Private/PluginManager.cpp
@@ -624,6 +624,15 @@ bool FPluginManager::ConfigureEnabledPlugins()
 					{
 						EditorTargetFileName = FTargetReceipt::GetDefaultPath(*FPaths::EngineDir(), *Target.Name, FPlatformProcess::GetBinariesSubdirectory(), FApp::GetBuildConfiguration(), nullptr);
 					}
+					else
+					{
+						// Unknown path, possibly built on another machine. Try project first with this target name, then engine
+						EditorTargetFileName = FTargetReceipt::GetDefaultPath(FPlatformMisc::ProjectDir(), *Target.Name, FPlatformProcess::GetBinariesSubdirectory(), FApp::GetBuildConfiguration(), nullptr);
+						if (!FPaths::FileExists(EditorTargetFileName))
+						{
+							EditorTargetFileName = FTargetReceipt::GetDefaultPath(*FPaths::EngineDir(), *Target.Name, FPlatformProcess::GetBinariesSubdirectory(), FApp::GetBuildConfiguration(), nullptr);
+						}
+					}
 					break;
 				}
 			}
