- Fixed up some known config file issues (no longer need to manually manage known config names in UObject::GetConfigName)

diff --git a/Engine/Source/Runtime/Core/Private/Misc/ConfigCacheIni.cpp b/Engine/Source/Runtime/Core/Private/Misc/ConfigCacheIni.cpp
--- a/Engine/Source/Runtime/Core/Private/Misc/ConfigCacheIni.cpp
+++ b/Engine/Source/Runtime/Core/Private/Misc/ConfigCacheIni.cpp
@@ -3945,6 +3945,10 @@ bool FConfigCacheIni::InitializeKnownConfigFiles(const TCHAR* PlatformName, bool
 	return bEngineConfigCreated;
 }
 
+bool FConfigCacheIni::IsKnownConfigName(FName ConfigName)
+{
+	return KnownFiles.GetFile(ConfigName) != nullptr;
+}
 
 const FConfigFile* FConfigCacheIni::FKnownConfigFiles::GetFile(FName Name)
 {
diff --git a/Engine/Source/Runtime/Core/Public/Misc/ConfigCacheIni.h b/Engine/Source/Runtime/Core/Public/Misc/ConfigCacheIni.h
--- a/Engine/Source/Runtime/Core/Public/Misc/ConfigCacheIni.h
+++ b/Engine/Source/Runtime/Core/Public/Misc/ConfigCacheIni.h
@@ -1014,6 +1014,12 @@ public:
 	 */
 	bool InitializeKnownConfigFiles(const TCHAR* PlatformName, bool bDefaultEngineIniRequired);
 
+	/**
+	 * Returns true if the given name is one of the known configs, where the matching G****Ini property is going to match the 
+	 * base name ("Engine" returns true, which means GEngineIni's value is just "Engine")
+	 */
+	bool IsKnownConfigName(FName ConfigName);
+
 	/**
 	 * Create GConfig from a saved file
 	 */
diff --git a/Engine/Source/Runtime/CoreUObject/Private/UObject/Class.cpp b/Engine/Source/Runtime/CoreUObject/Private/UObject/Class.cpp
--- a/Engine/Source/Runtime/CoreUObject/Private/UObject/Class.cpp
+++ b/Engine/Source/Runtime/CoreUObject/Private/UObject/Class.cpp
@@ -5267,24 +5267,15 @@ void UClass::AssembleReferenceTokenStreams()
 
 const FString UClass::GetConfigName() const
 {
-	static const FName NAME_DeviceProfiles("DeviceProfiles");
-
-	if( ClassConfigName == NAME_Engine )
+	// Look up in the known configs
+	if (GConfig->IsKnownConfigName(ClassConfigName))
 	{
-		return GEngineIni;
+		return ClassConfigName.ToString();
 	}
 	else if( ClassConfigName == NAME_Editor )
 	{
 		return GEditorIni;
 	}
-	else if( ClassConfigName == NAME_Input )
-	{
-		return GInputIni;
-	}
-	else if( ClassConfigName == NAME_Game )
-	{
-		return GGameIni;
-	}
 	else if ( ClassConfigName == NAME_EditorSettings )
 	{
 		return GEditorSettingsIni;
@@ -5297,19 +5288,11 @@ const FString UClass::GetConfigName() const
 	{
 		return GEditorKeyBindingsIni;
 	}
-	else if (ClassConfigName == NAME_DeviceProfiles)
-	{
-		return GDeviceProfilesIni;
-	}
 	else if( ClassConfigName == NAME_None )
 	{
 		UE_LOG(LogClass, Fatal,TEXT("UObject::GetConfigName() called on class with config name 'None'. Class flags = 0x%08X"), (uint32)ClassFlags );
 		return TEXT("");
 	}
-	else if (ClassConfigName == NAME_GameUserSettings)
-	{
-		return GGameUserSettingsIni;
-	}
 	else
 	{
 		// generate the class ini name, and make sure it's up to date
diff --git a/Engine/Source/Runtime/Projects/Private/PluginManager.cpp b/Engine/Source/Runtime/Projects/Private/PluginManager.cpp
--- a/Engine/Source/Runtime/Projects/Private/PluginManager.cpp
+++ b/Engine/Source/Runtime/Projects/Private/PluginManager.cpp
@@ -872,9 +872,21 @@ bool FPluginManager::ConfigureEnabledPlugins()
 					for (const FString& ConfigFile : PluginConfigs)
 					{
 						FString PlaformName = FPlatformProperties::PlatformName();
-						// Use GetDestIniFilename to find the proper config file to combine into, since it manages command line overrides and path sanitization
-						PluginConfigFilename = FConfigCacheIni::GetDestIniFilename(*FPaths::GetBaseFilename(ConfigFile), *PlaformName, *FPaths::GeneratedConfigDir());
-						FConfigFile* FoundConfig = GConfig->Find(PluginConfigFilename);
+						// try a known ini name first
+						PluginConfigFilename = FPaths::GetBaseFilename(ConfigFile);
+						FConfigFile* FoundConfig;
+						if (GConfig->IsKnownConfigName(*PluginConfigFilename))
+						{
+							FoundConfig = GConfig->FindConfigFile(PluginConfigFilename);
+						}
+						// if it's not a known type, generate it final name
+						else
+						{
+							// Use GetDestIniFilename to find the proper config file to combine into, since it manages command line overrides and path sanitization
+							PluginConfigFilename = FConfigCacheIni::GetDestIniFilename(*FPaths::GetBaseFilename(ConfigFile), *PlaformName, *FPaths::GeneratedConfigDir());
+							FoundConfig = GConfig->Find(PluginConfigFilename);
+						}
+
 						if (FoundConfig != nullptr)
 						{
 							UE_LOG(LogPluginManager, Log, TEXT("Found config from plugin[%s] %s"), *Plugin.GetName(), *PluginConfigFilename);
