Add IPluginmanager::GetLastCompletedLoadingPhase to support registering with OnLoadingPhaseComplete only if required loading phase has not yet happened.
#b Graeme.Thornton
#rnx

diff --git a/Engine/Source/Runtime/Projects/Private/PluginManager.cpp b/Engine/Source/Runtime/Projects/Private/PluginManager.cpp
--- a/Engine/Source/Runtime/Projects/Private/PluginManager.cpp
+++ b/Engine/Source/Runtime/Projects/Private/PluginManager.cpp
@@ -1509,32 +1509,38 @@ bool FPluginManager::TryLoadModulesForPlugin( const FPlugin& Plugin, const ELoad
 bool FPluginManager::LoadModulesForEnabledPlugins( const ELoadingPhase::Type LoadingPhase )
 {
 	// Figure out which plugins are enabled
-	if(!ConfigureEnabledPlugins())
+	bool bSuccess = true;
+	if (!ConfigureEnabledPlugins())
 	{
-		LoadingPhaseCompleteEvent.Broadcast(LoadingPhase, /*bSuccess =*/false);
-		return false;
+		bSuccess = false;
 	}
-
-	FScopedSlowTask SlowTask(AllPlugins.Num());
-
-	// Load plugins!
-	for( const TPair<FString, TSharedRef< FPlugin >>& PluginPair : AllPlugins )
+	else
 	{
-		const TSharedRef<FPlugin> &Plugin = PluginPair.Value;
+		FScopedSlowTask SlowTask(AllPlugins.Num());
 
-		SlowTask.EnterProgressFrame(1);
-
-		if ( Plugin->bEnabled && !Plugin->Descriptor.bExplicitlyLoaded )
+		// Load plugins!
+		for (const TPair<FString, TSharedRef< FPlugin >>& PluginPair : AllPlugins)
 		{
-			if (!TryLoadModulesForPlugin(Plugin.Get(), LoadingPhase))
+			const TSharedRef<FPlugin>& Plugin = PluginPair.Value;
+
+			SlowTask.EnterProgressFrame(1);
+
+			if (Plugin->bEnabled && !Plugin->Descriptor.bExplicitlyLoaded)
 			{
-				LoadingPhaseCompleteEvent.Broadcast(LoadingPhase, /*bSuccess =*/false);
-				return false;
+				if (!TryLoadModulesForPlugin(Plugin.Get(), LoadingPhase))
+				{
+					bSuccess = false;
+					break;
+				}
 			}
 		}
 	}
-	LoadingPhaseCompleteEvent.Broadcast(LoadingPhase, /*bSuccess =*/true);
-	return true;
+	UE_CLOG(LastCompletedLoadingPhase != ELoadingPhase::None && LastCompletedLoadingPhase >= LoadingPhase,
+		LogPluginManager, Error, TEXT("LoadModulesForEnabledPlugins called on phase %d after already being called on later or equal phase %d."),
+		static_cast<int32>(LoadingPhase), static_cast<int32>(LastCompletedLoadingPhase));
+	LastCompletedLoadingPhase = LoadingPhase;
+	LoadingPhaseCompleteEvent.Broadcast(LoadingPhase, bSuccess);
+	return bSuccess;
 }
 
 IPluginManager::FLoadingModulesForPhaseEvent& FPluginManager::OnLoadingPhaseComplete()
@@ -1542,6 +1548,11 @@ IPluginManager::FLoadingModulesForPhaseEvent& FPluginManager::OnLoadingPhaseComp
 	return LoadingPhaseCompleteEvent;
 }
 
+ELoadingPhase::Type FPluginManager::GetLastCompletedLoadingPhase() const
+{
+	return LastCompletedLoadingPhase;
+}
+
 void FPluginManager::GetLocalizationPathsForEnabledPlugins( TArray<FString>& OutLocResPaths )
 {
 	// Figure out which plugins are enabled
diff --git a/Engine/Source/Runtime/Projects/Private/PluginManager.h b/Engine/Source/Runtime/Projects/Private/PluginManager.h
--- a/Engine/Source/Runtime/Projects/Private/PluginManager.h
+++ b/Engine/Source/Runtime/Projects/Private/PluginManager.h
@@ -108,6 +108,7 @@ public:
 	virtual bool AddToPluginsList( const FString& PluginFilename ) override;
 	virtual bool LoadModulesForEnabledPlugins( const ELoadingPhase::Type LoadingPhase ) override;
 	virtual FLoadingModulesForPhaseEvent& OnLoadingPhaseComplete() override;
+	virtual ELoadingPhase::Type GetLastCompletedLoadingPhase() const override;
 	virtual void GetLocalizationPathsForEnabledPlugins( TArray<FString>& OutLocResPaths ) override;
 	virtual void SetRegisterMountPointDelegate( const FRegisterMountPointDelegate& Delegate ) override;
 	virtual void SetUnRegisterMountPointDelegate( const FRegisterMountPointDelegate& Delegate ) override;
@@ -233,6 +234,9 @@ private:
 
 	/** Callback for notifications that a loading phase was completed */
 	FLoadingModulesForPhaseEvent LoadingPhaseCompleteEvent;
+
+	/** The highest LoadingPhase that has so far completed */
+	ELoadingPhase::Type LastCompletedLoadingPhase = ELoadingPhase::None;
 };
 PRAGMA_ENABLE_DEPRECATION_WARNINGS
 
diff --git a/Engine/Source/Runtime/Projects/Public/Interfaces/IPluginManager.h b/Engine/Source/Runtime/Projects/Public/Interfaces/IPluginManager.h
--- a/Engine/Source/Runtime/Projects/Public/Interfaces/IPluginManager.h
+++ b/Engine/Source/Runtime/Projects/Public/Interfaces/IPluginManager.h
@@ -209,6 +209,9 @@ public:
 	 */
 	virtual bool LoadModulesForEnabledPlugins( const ELoadingPhase::Type LoadingPhase ) = 0;
 
+	/** Returns the highest loading phase that has so far completed */
+	virtual ELoadingPhase::Type GetLastCompletedLoadingPhase() const = 0;
+
 	/**
 	 * Callback for when modules for when LoadModulesForEnabledPlugins() completes loading for a specific phase.
 	 */
