- Fixed support for a project with all platforms supported, but with a Platforms subdirectory, would treat those Platforms as the onjly supported platforms

diff --git a/Engine/Source/Runtime/Projects/Private/ProjectDescriptor.cpp b/Engine/Source/Runtime/Projects/Private/ProjectDescriptor.cpp
--- a/Engine/Source/Runtime/Projects/Private/ProjectDescriptor.cpp
+++ b/Engine/Source/Runtime/Projects/Private/ProjectDescriptor.cpp
@@ -184,24 +184,30 @@ bool FProjectDescriptor::Read(const FJsonObject& Object, const FString& PathToPr
 	}
 
 	// check if the project has directories for extended platforms, and assume support if it does
-	TArray<FString> ExtendedPlatforms;
-	IFileManager::Get().IterateDirectory(*(FPaths::Combine(PathToProject, TEXT("Platforms"))), [&ExtendedPlatforms](const TCHAR* InFilenameOrDirectory, const bool bInIsDirectory) -> bool
+	// however, if there were no platforms already listed, then all platforms are supported, and we don't
+	// want to add a platform or two here, because then _only_ those platforms will be supported
+	// (empty TargetPlatforms array means _all_ platforms are supported)
+	if (TargetPlatforms.Num() > 0)
 	{
-		if (bInIsDirectory)
+		TArray<FString> ExtendedPlatforms;
+		IFileManager::Get().IterateDirectory(*(FPaths::Combine(PathToProject, TEXT("Platforms"))), [&ExtendedPlatforms](const TCHAR* InFilenameOrDirectory, const bool bInIsDirectory) -> bool
 		{
-			FString LastDirectory = FPaths::GetBaseFilename(FString(InFilenameOrDirectory));
-			ExtendedPlatforms.Emplace(LastDirectory);
-		}
-		return true;
-	});
+			if (bInIsDirectory)
+			{
+				FString LastDirectory = FPaths::GetBaseFilename(FString(InFilenameOrDirectory));
+				ExtendedPlatforms.Emplace(LastDirectory);
+			}
+			return true;
+		});
 
-	const TMap<FName, FDataDrivenPlatformInfo>& AllPlatformInfos = FDataDrivenPlatformInfoRegistry::GetAllPlatformInfos();
-	for (const FString& ExtendedPlatform : ExtendedPlatforms)
-	{
-		FName PlatformName(*ExtendedPlatform);
-		if (AllPlatformInfos.Contains(PlatformName))
+		const TMap<FName, FDataDrivenPlatformInfo>& AllPlatformInfos = FDataDrivenPlatformInfoRegistry::GetAllPlatformInfos();
+		for (const FString& ExtendedPlatform : ExtendedPlatforms)
 		{
-			TargetPlatforms.AddUnique(PlatformName);
+			FName PlatformName(*ExtendedPlatform);
+			if (AllPlatformInfos.Contains(PlatformName))
+			{
+				TargetPlatforms.AddUnique(PlatformName);
+			}
 		}
 	}
 
