Fixed duplicated "ConfigureEnabledPluginForCurrentTarget" cpu timer and added few more SCOPED_BOOT_TIMING in FPluginManager::ConfigureEnabledPlugins().

diff --git a/Engine/Source/Runtime/Projects/Private/PluginManager.cpp b/Engine/Source/Runtime/Projects/Private/PluginManager.cpp
--- a/Engine/Source/Runtime/Projects/Private/PluginManager.cpp
+++ b/Engine/Source/Runtime/Projects/Private/PluginManager.cpp
@@ -663,6 +663,8 @@ bool FPluginManager::ConfigureEnabledPlugins()
 
 		// Check which plugins have been enabled or excluded via the command line
 		{
+			SCOPED_BOOT_TIMING("ParseCmdLineForPlugins");
+
 			auto ParsePluginsList = [](const TCHAR* InListKey) -> TArray<FString>
 			{
 				TArray<FString> PluginsList;
@@ -722,6 +724,8 @@ bool FPluginManager::ConfigureEnabledPlugins()
 
 		if (!FParse::Param(FCommandLine::Get(), TEXT("NoEnginePlugins")))
 		{
+			SCOPED_BOOT_TIMING("EnginePlugins");
+
 #if READ_TARGET_ENABLED_PLUGINS_FROM_RECEIPT
 			// Configure the plugins that were enabled or disabled from the target file using the target receipt file
 			FString DefaultEditorTarget;
@@ -770,6 +774,7 @@ bool FPluginManager::ConfigureEnabledPlugins()
 			};
 
 			{
+				SCOPED_BOOT_TIMING("ConfigureTargetEnabledPluginsFromReceipt");
 				bool bErrorConfiguring = false;
 				if (!ConfigurePluginsFromFirstMatchingTargetFile(FPlatformMisc::ProjectDir(), bErrorConfiguring))
 				{
@@ -781,31 +786,38 @@ bool FPluginManager::ConfigureEnabledPlugins()
 				}
 			}
 #else
-			// Configure the plugins that were enabled from the target file using defines
-			TArray<FString> TargetEnabledPlugins = { UBT_TARGET_ENABLED_PLUGINS };
-			for (const FString& TargetEnabledPlugin : TargetEnabledPlugins)
 			{
-				if (!ConfiguredPluginNames.Contains(TargetEnabledPlugin))
+				// Configure the plugins that were enabled from the target file using defines
+				SCOPED_BOOT_TIMING("ConfigureTargetEnabledPlugins");
+
+				TArray<FString> TargetEnabledPlugins = { UBT_TARGET_ENABLED_PLUGINS };
+				for (const FString& TargetEnabledPlugin : TargetEnabledPlugins)
 				{
-					if (!ConfigureEnabledPluginForCurrentTarget(FPluginReferenceDescriptor(TargetEnabledPlugin, true), EnabledPlugins))
+					if (!ConfiguredPluginNames.Contains(TargetEnabledPlugin))
 					{
-						return false;
+						if (!ConfigureEnabledPluginForCurrentTarget(FPluginReferenceDescriptor(TargetEnabledPlugin, true), EnabledPlugins))
+						{
+							return false;
+						}
+						ConfiguredPluginNames.Add(TargetEnabledPlugin);
 					}
-					ConfiguredPluginNames.Add(TargetEnabledPlugin);
 				}
 			}
-
-			// Configure the plugins that were disabled from the target file using defines
-			TArray<FString> TargetDisabledPlugins = { UBT_TARGET_DISABLED_PLUGINS };
-			for (const FString& TargetDisabledPlugin : TargetDisabledPlugins)
 			{
-				if (!ConfiguredPluginNames.Contains(TargetDisabledPlugin))
+				// Configure the plugins that were disabled from the target file using defines
+				SCOPED_BOOT_TIMING("ConfigureTargetDisabledPlugins");
+
+				TArray<FString> TargetDisabledPlugins = { UBT_TARGET_DISABLED_PLUGINS };
+				for (const FString& TargetDisabledPlugin : TargetDisabledPlugins)
 				{
-					if (!ConfigureEnabledPluginForCurrentTarget(FPluginReferenceDescriptor(TargetDisabledPlugin, false), EnabledPlugins))
+					if (!ConfiguredPluginNames.Contains(TargetDisabledPlugin))
 					{
-						return false;
+						if (!ConfigureEnabledPluginForCurrentTarget(FPluginReferenceDescriptor(TargetDisabledPlugin, false), EnabledPlugins))
+						{
+							return false;
+						}
+						ConfiguredPluginNames.Add(TargetDisabledPlugin);
 					}
-					ConfiguredPluginNames.Add(TargetDisabledPlugin);
 				}
 			}
 #endif // READ_TARGET_ENABLED_PLUGINS_FROM_RECEIPT
@@ -814,10 +826,10 @@ bool FPluginManager::ConfigureEnabledPlugins()
 			// Find all the plugin references in the project file
 			const FProjectDescriptor* ProjectDescriptor = IProjectManager::Get().GetCurrentProject();
 			{
-				SCOPED_BOOT_TIMING("ConfigureEnabledPluginForCurrentTarget");
+				SCOPED_BOOT_TIMING("AddPluginReferences");
+
 				if (ProjectDescriptor != nullptr)
 				{
-
 					bAllowEnginePluginsEnabledByDefault = !ProjectDescriptor->bDisableEnginePluginsByDefault;
 
 					// Copy the plugin references, since we may modify the project if any plugins are missing
@@ -836,60 +848,73 @@ bool FPluginManager::ConfigureEnabledPlugins()
 				}
 			}
 
-			// Add the plugins which are enabled by default
-			for(const FString& PluginName : PluginsToConfigure)
 			{
-				const TSharedRef<FPlugin>& Plugin = AllPlugins.FindChecked(PluginName);
-				if (Plugin->IsEnabledByDefault(bAllowEnginePluginsEnabledByDefault) && !ConfiguredPluginNames.Contains(PluginName))
+				// Add the plugins which are enabled by default
+				SCOPED_BOOT_TIMING("AddPluginsEnabledByDefault");
+
+				for (const FString& PluginName : PluginsToConfigure)
 				{
-					if (!ConfigureEnabledPluginForCurrentTarget(FPluginReferenceDescriptor(PluginName, true), EnabledPlugins))
+					const TSharedRef<FPlugin>& Plugin = AllPlugins.FindChecked(PluginName);
+					if (Plugin->IsEnabledByDefault(bAllowEnginePluginsEnabledByDefault) && !ConfiguredPluginNames.Contains(PluginName))
 					{
-						return false;
+						if (!ConfigureEnabledPluginForCurrentTarget(FPluginReferenceDescriptor(PluginName, true), EnabledPlugins))
+						{
+							return false;
+						}
+						ConfiguredPluginNames.Add(PluginName);
 					}
-					ConfiguredPluginNames.Add(PluginName);
 				}
 			}
 		}
-#if IS_PROGRAM
-		// Programs can also define the list of enabled plugins in ini
-		TArray<FString> ProgramPluginNames;
-		GConfig->GetArray(TEXT("Plugins"), TEXT("ProgramEnabledPlugins"), ProgramPluginNames, GEngineIni);
 
-		for (const FString& PluginName : ProgramPluginNames)
+#if IS_PROGRAM
 		{
-			if (!ConfiguredPluginNames.Contains(PluginName))
+			// Programs can also define the list of enabled plugins in ini
+			SCOPED_BOOT_TIMING("AddProgramEnabledPlugins");
+
+			TArray<FString> ProgramPluginNames;
+			GConfig->GetArray(TEXT("Plugins"), TEXT("ProgramEnabledPlugins"), ProgramPluginNames, GEngineIni);
+
+			for (const FString& PluginName : ProgramPluginNames)
 			{
-				if (!ConfigureEnabledPluginForCurrentTarget(FPluginReferenceDescriptor(PluginName, true), EnabledPlugins))
+				if (!ConfiguredPluginNames.Contains(PluginName))
 				{
-					return false;
+					if (!ConfigureEnabledPluginForCurrentTarget(FPluginReferenceDescriptor(PluginName, true), EnabledPlugins))
+					{
+						return false;
+					}
+					ConfiguredPluginNames.Add(PluginName);
 				}
-				ConfiguredPluginNames.Add(PluginName);
 			}
 		}
 #endif
 
-		// Mark all the plugins as enabled
-		for (TPair<FString, FPlugin*>& Pair : EnabledPlugins)
 		{
-			FPlugin& Plugin = *Pair.Value;
+			// Mark all the plugins as enabled
+			SCOPED_BOOT_TIMING("MarkEnabledPlugins");
 
-#if !IS_MONOLITHIC
-			// Mount the binaries directory, and check the modules are valid
-			if (Plugin.Descriptor.Modules.Num() > 0)
+			for (TPair<FString, FPlugin*>& Pair : EnabledPlugins)
 			{
-				// Mount the binaries directory
-				const FString PluginBinariesPath = FPaths::Combine(*FPaths::GetPath(Plugin.FileName), TEXT("Binaries"), FPlatformProcess::GetBinariesSubdirectory());
-				FModuleManager::Get().AddBinariesDirectory(*PluginBinariesPath, Plugin.GetLoadedFrom() == EPluginLoadedFrom::Project);
-			}
+				FPlugin& Plugin = *Pair.Value;
 
-			// Check the declared engine version. This is a soft requirement, so allow the user to skip over it.
-			if (!IsPluginCompatible(Plugin) && !PromptToLoadIncompatiblePlugin(Plugin))
-			{
-				UE_LOG(LogPluginManager, Display, TEXT("Skipping load of '%s'."), *Plugin.Name);
-				continue;
-			}
+#if !IS_MONOLITHIC
+				// Mount the binaries directory, and check the modules are valid
+				if (Plugin.Descriptor.Modules.Num() > 0)
+				{
+					// Mount the binaries directory
+					const FString PluginBinariesPath = FPaths::Combine(*FPaths::GetPath(Plugin.FileName), TEXT("Binaries"), FPlatformProcess::GetBinariesSubdirectory());
+					FModuleManager::Get().AddBinariesDirectory(*PluginBinariesPath, Plugin.GetLoadedFrom() == EPluginLoadedFrom::Project);
+				}
+
+				// Check the declared engine version. This is a soft requirement, so allow the user to skip over it.
+				if (!IsPluginCompatible(Plugin) && !PromptToLoadIncompatiblePlugin(Plugin))
+				{
+					UE_LOG(LogPluginManager, Display, TEXT("Skipping load of '%s'."), *Plugin.Name);
+					continue;
+				}
 #endif
-			Plugin.bEnabled = true;
+				Plugin.bEnabled = true;
+			}
 		}
 
 		// If we made it here, we have all the required plugins
