Added an event for when plugins are edited and fixed some comments in IPluginManager.h

diff --git a/Engine/Source/Runtime/Projects/Private/PluginManager.cpp b/Engine/Source/Runtime/Projects/Private/PluginManager.cpp
--- a/Engine/Source/Runtime/Projects/Private/PluginManager.cpp
+++ b/Engine/Source/Runtime/Projects/Private/PluginManager.cpp
@@ -151,6 +151,13 @@ bool FPlugin::UpdateDescriptor(const FPluginDescriptor& NewDescriptor, FText& Ou
 	}
 
 	Descriptor = NewDescriptor;
+
+	IPluginManager& PluginManager = IPluginManager::Get();
+	if (PluginManager.OnPluginEdited().IsBound())
+	{
+		PluginManager.OnPluginEdited().Broadcast(*this);
+	}
+
 	return true;
 }
 
@@ -1766,6 +1773,11 @@ IPluginManager::FNewPluginMountedEvent& FPluginManager::OnNewPluginMounted()
 	return NewPluginMountedEvent;
 }
 
+IPluginManager::FNewPluginMountedEvent& FPluginManager::OnPluginEdited()
+{
+	return PluginEditedEvent;
+}
+
 void FPluginManager::MountNewlyCreatedPlugin(const FString& PluginName)
 {
 	TSharedPtr<FPlugin> Plugin = FindPluginInstance(PluginName);
diff --git a/Engine/Source/Runtime/Projects/Private/PluginManager.h b/Engine/Source/Runtime/Projects/Private/PluginManager.h
--- a/Engine/Source/Runtime/Projects/Private/PluginManager.h
+++ b/Engine/Source/Runtime/Projects/Private/PluginManager.h
@@ -127,6 +127,7 @@ public:
 	virtual TArray<TSharedRef<IPlugin>> GetPluginsWithPakFile() const override;
 	virtual FNewPluginMountedEvent& OnNewPluginCreated() override;
 	virtual FNewPluginMountedEvent& OnNewPluginMounted() override;
+	virtual FNewPluginMountedEvent& OnPluginEdited() override;
 	virtual void MountNewlyCreatedPlugin(const FString& PluginName) override;
 	virtual void MountExplicitlyLoadedPlugin(const FString& PluginName) override;
 	virtual bool UnmountExplicitlyLoadedPlugin(const FString& PluginName, FText* OutReason) override;
@@ -231,6 +232,7 @@ private:
 	/** Callback for notifications that a new plugin was mounted */
 	FNewPluginMountedEvent NewPluginCreatedEvent;
 	FNewPluginMountedEvent NewPluginMountedEvent;
+	FNewPluginMountedEvent PluginEditedEvent;
 
 	/** Callback for notifications that a loading phase was completed */
 	FLoadingModulesForPhaseEvent LoadingPhaseCompleteEvent;
diff --git a/Engine/Source/Runtime/Projects/Public/Interfaces/IPluginManager.h b/Engine/Source/Runtime/Projects/Public/Interfaces/IPluginManager.h
--- a/Engine/Source/Runtime/Projects/Public/Interfaces/IPluginManager.h
+++ b/Engine/Source/Runtime/Projects/Public/Interfaces/IPluginManager.h
@@ -335,14 +335,21 @@ public:
 	 */
 	DECLARE_EVENT_OneParam(IPluginManager, FNewPluginMountedEvent, IPlugin&);
 
-
+	/**
+	 * Event signature for being notified that a new plugin has been created
+	 */
 	virtual FNewPluginMountedEvent& OnNewPluginCreated() = 0;
 
 	/**
-	 * Gets an array of plugins that loaded their own content pak file
+	 * Event for being notified that a new plugin has been mounted
 	 */
 	virtual FNewPluginMountedEvent& OnNewPluginMounted() = 0;
 
+	/**
+	 * Event for being notified that a plugin has been edited
+	 */
+	virtual FNewPluginMountedEvent& OnPluginEdited() = 0;
+
 	/**
 	 * Marks a newly created plugin as enabled, mounts its content and tries to load its modules
 	 */
