Do not log an error if LoadModulesForEnabledPlugins is called multiple times in a row on the same phase, since this occurs for PreEarlyLoadingScreen.
#rnx

diff --git a/Engine/Source/Runtime/Projects/Private/PluginManager.cpp b/Engine/Source/Runtime/Projects/Private/PluginManager.cpp
--- a/Engine/Source/Runtime/Projects/Private/PluginManager.cpp
+++ b/Engine/Source/Runtime/Projects/Private/PluginManager.cpp
@@ -1567,9 +1567,13 @@ bool FPluginManager::LoadModulesForEnabledPlugins( const ELoadingPhase::Type Loa
 			}
 		}
 	}
-	UE_CLOG(LastCompletedLoadingPhase != ELoadingPhase::None && LastCompletedLoadingPhase >= LoadingPhase,
-		LogPluginManager, Error, TEXT("LoadModulesForEnabledPlugins called on phase %d after already being called on later or equal phase %d."),
+	// Some phases such as ELoadingPhase::PreEarlyLoadingScreen are potentially called multiple times,
+	// but we do not return to an earlier phase after calling LoadModulesForEnabledPlugins on a later phase
+	UE_CLOG(LastCompletedLoadingPhase != ELoadingPhase::None && LastCompletedLoadingPhase > LoadingPhase,
+		LogPluginManager, Error, TEXT("LoadModulesForEnabledPlugins called on phase %d after already being called on later phase %d."),
 		static_cast<int32>(LoadingPhase), static_cast<int32>(LastCompletedLoadingPhase));
+
+	// We send the broadcast event each time, even if this function is called multiple times with the same phase
 	LastCompletedLoadingPhase = LoadingPhase;
 	LoadingPhaseCompleteEvent.Broadcast(LoadingPhase, bSuccess);
 	return bSuccess;
