Fixed some typos in the comment for bHasExplicitPlatforms
#rnx

diff --git a/Engine/Source/Runtime/Projects/Public/PluginDescriptor.h b/Engine/Source/Runtime/Projects/Public/PluginDescriptor.h
--- a/Engine/Source/Runtime/Projects/Public/PluginDescriptor.h
+++ b/Engine/Source/Runtime/Projects/Public/PluginDescriptor.h
@@ -105,7 +105,7 @@ struct PROJECTS_API FPluginDescriptor
 	/** When true, this plugin's modules will not be loaded automatically nor will it's content be mounted automatically. It will load/mount when explicitly requested and LoadingPhases will be ignored */
 	bool bExplicitlyLoaded;
 
-	/** When true, an empty SupportedTargetPlatforms is interpeted as 'no platforms' with the expectation that explict platforms will be added in plugin platform extensions */
+	/** When true, an empty SupportedTargetPlatforms is interpreted as 'no platforms' with the expectation that explicit platforms will be added in plugin platform extensions */
 	bool bHasExplicitPlatforms;
 
 	/** If true, this plugin from a platform extension extending another plugin */
