Add GetConfigFilename to config cache that handles going from either Engine or CustomSystem to the right key filename
Fix spelling error and comments while I'm already modifying this code
Change calling sites that were doing this manually to use new function, and fix a bad merge in the plugin code
Add logic to strip Base and Default prefixes from plugin config overrides, as the old behavior with Find instead of FindConfigFile lead to that sometimes working by accident
I may update the handling of Base/DefaultPluginName.ini in a separate checkin to make that more consistent

diff --git a/Engine/Plugins/Online/OnlineFramework/Source/Hotfix/Private/OnlineHotfixManager.cpp b/Engine/Plugins/Online/OnlineFramework/Source/Hotfix/Private/OnlineHotfixManager.cpp
--- a/Engine/Plugins/Online/OnlineFramework/Source/Hotfix/Private/OnlineHotfixManager.cpp
+++ b/Engine/Plugins/Online/OnlineFramework/Source/Hotfix/Private/OnlineHotfixManager.cpp
@@ -873,16 +873,7 @@ FString UOnlineHotfixManager::BuildConfigCacheKey(const FString& IniName)
 {
 	const FString IniNameNoExtension = FPaths::GetBaseFilename(IniName);
 
-	// Known ini files such as Engine, Game, etc.. are referred to as just the name with no extension within the config system.
-	if (GConfig->IsKnownConfigName(FName(*IniNameNoExtension, FNAME_Find)))
-	{
-		return IniNameNoExtension;
-	}
-	else
-	{
-		// Non-known ini files are looked up using their full path.
-		return GConfig->GetDestIniFilename(*IniNameNoExtension, nullptr, *FPaths::GeneratedConfigDir());
-	}
+	return GConfig->GetConfigFilename(*IniNameNoExtension);
 }
 
 FConfigFile* UOnlineHotfixManager::GetConfigFile(const FString& IniName)
diff --git a/Engine/Source/Developer/TraceInsights/Private/Insights/TraceInsightsModule.cpp b/Engine/Source/Developer/TraceInsights/Private/Insights/TraceInsightsModule.cpp
--- a/Engine/Source/Developer/TraceInsights/Private/Insights/TraceInsightsModule.cpp
+++ b/Engine/Source/Developer/TraceInsights/Private/Insights/TraceInsightsModule.cpp
@@ -57,7 +57,7 @@ void FTraceInsightsModule::StartupModule()
 	RegisterComponent(Insights::FTaskGraphProfilerManager::CreateInstance());
 	RegisterComponent(Insights::FContextSwitchesProfilerManager::CreateInstance());
 
-	UnrealInsightsLayoutIni = FConfigCacheIni::GetDestIniFilename(TEXT("UnrealInsightsLayout"), nullptr, *FPaths::GeneratedConfigDir());
+	UnrealInsightsLayoutIni = GConfig->GetConfigFilename(TEXT("UnrealInsightsLayout"));
 }
 
 ////////////////////////////////////////////////////////////////////////////////////////////////////
diff --git a/Engine/Source/Runtime/Core/Private/Misc/ConfigCacheIni.cpp b/Engine/Source/Runtime/Core/Private/Misc/ConfigCacheIni.cpp
--- a/Engine/Source/Runtime/Core/Private/Misc/ConfigCacheIni.cpp
+++ b/Engine/Source/Runtime/Core/Private/Misc/ConfigCacheIni.cpp
@@ -1291,7 +1291,7 @@ void FConfigFile::OverrideFromCommandline(FConfigFile* File, const FString& File
 }
 
 
-void FConfigFile::AddDynamicLayerToHeirarchy(const FString& Filename)
+void FConfigFile::AddDynamicLayerToHierarchy(const FString& Filename)
 {
 	FString ConfigContent;
 	if (!FFileHelper::LoadFileToString(ConfigContent, *Filename))
@@ -2784,6 +2784,21 @@ bool FConfigCacheIni::EmptySectionsMatchingString( const TCHAR* SectionString, c
 	return bEmptied;
 }
 
+FString FConfigCacheIni::GetConfigFilename(const TCHAR* BaseIniName)
+{
+	// Known ini files such as Engine, Game, etc.. are referred to as just the name with no extension within the config system.
+	if (IsKnownConfigName(FName(BaseIniName, FNAME_Find)))
+	{
+		return FString(BaseIniName);
+	}
+	else
+	{
+		// Non-known ini files are looked up using their full path
+		// This always uses the default platform as non-known files are not valid for other platforms
+		return GetDestIniFilename(BaseIniName, nullptr, *FPaths::GeneratedConfigDir());
+	}
+}
+
 /**
  * Retrieve a list of all of the config files stored in the cache
  *
diff --git a/Engine/Source/Runtime/Core/Public/Misc/ConfigCacheIni.h b/Engine/Source/Runtime/Core/Public/Misc/ConfigCacheIni.h
--- a/Engine/Source/Runtime/Core/Public/Misc/ConfigCacheIni.h
+++ b/Engine/Source/Runtime/Core/Public/Misc/ConfigCacheIni.h
@@ -568,8 +568,11 @@ public:
 	/** Checks the command line for any overridden config file settings */
 	CORE_API static bool OverrideFileFromCommandline(FString& Filename);
 
-	/** Appends a new INI file to the SourceIniHierarchy and combines it */
-	CORE_API void AddDynamicLayerToHeirarchy(const FString& Filename);
+	/** Appends a new INI file to the SourceIniHierarchy and combines it with the current contents */
+	CORE_API void AddDynamicLayerToHierarchy(const FString& Filename);
+
+	UE_DEPRECATED(5.0, "Call AddDynamicLayerToHierarchy. You also may need to call GetConfigFilename to get the right FConfigFile")
+	CORE_API void AddDynamicLayerToHeirarchy(const FString& Filename) { AddDynamicLayerToHierarchy(Filename); }
 
 	friend FArchive& operator<<(FArchive& Ar, FConfigFile& ConfigFile);
 private:
@@ -698,8 +701,22 @@ public:
 	*/
 	virtual void Parse1ToNSectionOfNames(const TCHAR* Section, const TCHAR* KeyOne, const TCHAR* KeyN, TMap<FName, TArray<FName> >& OutMap, const FString& Filename);
 
-	/** Finds Config file based on the final, generated ini name */
-	FConfigFile* FindConfigFile( const FString& Filename );
+	/**
+	 * Finds the in-memory config file for a config cache filename.
+	 *
+	 * @param A known key like GEngineIni, or the return value of GetConfigFilename
+	 *
+	 * @return The existing config file or null if it does not exist in memory
+	 */
+	FConfigFile* FindConfigFile(const FString& Filename);
+
+	/**
+	 * Finds, loads, or creates the in-memory config file for a config cache filename.
+	 * 
+	 * @param A known key like GEngineIni, or the return value of GetConfigFilename
+	 * 
+	 * @return A new or existing config file
+	 */
 	FConfigFile* Find(const FString& InFilename);
 
 	/**
@@ -709,7 +726,7 @@ public:
 	 */
 	bool ContainsConfigFile(const FConfigFile* ConfigFile) const;
 
-	UE_DEPRECATED(5.0, "CreateIfNotFound is deprecated, please use the overload without this parameter")
+	UE_DEPRECATED(5.0, "CreateIfNotFound is deprecated, please use the overload without this parameter or FindConfigFile")
 	FConfigFile* Find(const FString& Filename, bool CreateIfNotFound);
 
 	/** Finds Config file that matches the base name such as "Engine" */
@@ -749,6 +766,16 @@ public:
 	bool EmptySection( const TCHAR* Section, const FString& Filename );
 	bool EmptySectionsMatchingString( const TCHAR* SectionString, const FString& Filename );
 
+	/**
+	 * For a base ini name, gets the config cache filename key that is used by other functions like Find.
+	 * This will be the base name for known configs like Engine and the destination filename for others.
+	 *
+	 * @param IniBaseName Base name of the .ini (Engine, Game, CustomSystem)
+	 *
+	 * @return Filename key used by other cache functions
+	 */
+	FString GetConfigFilename(const TCHAR* BaseIniName);
+
 	/**
 	 * Retrieve a list of all of the config files stored in the cache
 	 *
diff --git a/Engine/Source/Runtime/Projects/Private/PluginManager.cpp b/Engine/Source/Runtime/Projects/Private/PluginManager.cpp
--- a/Engine/Source/Runtime/Projects/Private/PluginManager.cpp
+++ b/Engine/Source/Runtime/Projects/Private/PluginManager.cpp
@@ -625,14 +625,20 @@ bool FPluginManager::IntegratePluginsIntoConfig(FConfigCacheIni& ConfigSystem, c
 			IFileManager::Get().FindFiles(PluginConfigs, *PluginConfigDir, TEXT("ini"));
 			for (const FString& ConfigFile : PluginConfigs)
 			{
-				// Use GetDestIniFilename to find the proper config file to combine into, since it manages command line overrides and path sanitization
-				FString PluginConfigFilename = FConfigCacheIni::GetDestIniFilename(*FPaths::GetBaseFilename(ConfigFile), PlatformName, *FPaths::GeneratedConfigDir());
+				// Optionally strip Default and Base from source filename, old behavior would create a new misnamed file which allowed redirects to work
+				// But, it's better to just apply it to the config actually desired by the user
+				FString BaseConfigFile = *FPaths::GetBaseFilename(ConfigFile);
+				BaseConfigFile.RemoveFromStart(TEXT("Base"));
+				BaseConfigFile.RemoveFromStart(TEXT("Default"));
+
+				// Use GetConfigFilename to find the proper config file to combine into, since it manages command line overrides and path sanitization
+				FString PluginConfigFilename = ConfigSystem.GetConfigFilename(*BaseConfigFile);
 				FConfigFile* FoundConfig = ConfigSystem.FindConfigFile(PluginConfigFilename);
 				if (FoundConfig != nullptr)
 				{
 					UE_LOG(LogPluginManager, Log, TEXT("Found config from plugin[%s] %s"), *Plugin.GetName(), *PluginConfigFilename);
 
-					FoundConfig->AddDynamicLayerToHeirarchy(FPaths::Combine(PluginConfigDir, ConfigFile));
+					FoundConfig->AddDynamicLayerToHierarchy(FPaths::Combine(PluginConfigDir, ConfigFile));
 				}
 			}
 
@@ -1009,8 +1015,8 @@ bool FPluginManager::ConfigureEnabledPlugins()
 					SourceConfigDir = PluginConfigDir;
 				}
 
-				FString PluginConfigFilename = FString::Printf(TEXT("%s%s/%s.ini"), *FPaths::GeneratedConfigDir(), *PlatformName, *Plugin.Name);
-				FPaths::MakeStandardFilename(PluginConfigFilename); // This needs to match what we do in ConfigCacheIni.cpp's GetDestIniFilename method. Otherwise, the hash results will differ and the plugin's version will be overwritten later.
+				// Build the config system key for PluginName.ini
+				FString PluginConfigFilename = GConfig->GetConfigFilename(*Plugin.Name);
 				{
 					FScopeLock Locker(&ConfigCS);
 
@@ -1029,30 +1035,23 @@ bool FPluginManager::ConfigureEnabledPlugins()
 				IFileManager::Get().FindFiles(PluginConfigs, *PluginConfigDir, TEXT("ini"));
 				for (const FString& ConfigFile : PluginConfigs)
 				{
-					// Use GetDestIniFilename to find the proper config file to combine into, since it manages command line overrides and path sanitization
-					PluginConfigFilename = FConfigCacheIni::GetDestIniFilename(*FPaths::GetBaseFilename(ConfigFile), *PlatformName, *FPaths::GeneratedConfigDir());
+					// Optionally strip Default and Base from source filename, old behavior would create a new misnamed file which allowed redirects to work
+					// But, it's better to just apply it to the config actually desired by the user
+					FString BaseConfigFile = *FPaths::GetBaseFilename(ConfigFile);					
+					BaseConfigFile.RemoveFromStart(TEXT("Base"));
+					BaseConfigFile.RemoveFromStart(TEXT("Default"));
+
+					// Build the config system key for the overridden config
+					PluginConfigFilename = GConfig->GetConfigFilename(*BaseConfigFile);
 					{
 						FScopeLock Locker(&ConfigCS);
-						FConfigFile* FoundConfig = GConfig->Find(PluginConfigFilename);
-						if (FoundConfig != nullptr)
-						{
-						    if (GConfig->IsKnownConfigName(*PluginConfigFilename))
-						    {
-							    FoundConfig = GConfig->FindConfigFile(PluginConfigFilename);
-						    }
-						    // if it's not a known type, generate it final name
-						    else
-						    {
-							    // Use GetDestIniFilename to find the proper config file to combine into, since it manages command line overrides and path sanitization
-							    PluginConfigFilename = FConfigCacheIni::GetDestIniFilename(*FPaths::GetBaseFilename(ConfigFile), *PlatformName, *FPaths::GeneratedConfigDir());
-							    FoundConfig = GConfig->FindConfigFile(PluginConfigFilename);
-						    }
-						}
+						FConfigFile* FoundConfig = GConfig->FindConfigFile(PluginConfigFilename);
+
 						if (FoundConfig != nullptr)
 						{
 							UE_LOG(LogPluginManager, Log, TEXT("Found config from plugin[%s] %s"), *Plugin.GetName(), *PluginConfigFilename);
 
-							FoundConfig->AddDynamicLayerToHeirarchy(FPaths::Combine(PluginConfigDir, ConfigFile));
+							FoundConfig->AddDynamicLayerToHierarchy(FPaths::Combine(PluginConfigDir, ConfigFile));
 
 #if ALLOW_INI_OVERRIDE_FROM_COMMANDLINE
 							// Don't allow plugins to stomp command line overrides, so re-apply them
