Fix game feature issue with previous checkin, the game feature code was still looking for the old plugin location so it would end up writing out a corrupted partial ini to the saved folder
Remove code that was stripping Base and Engine as it could have other similar side effects in out of date config code
This code might need further updates for UE5, the only reason it worked before was because it was depending on the fact that Find would create a nonsaveable config file when passed a bad key

diff --git a/Engine/Plugins/Experimental/GameFeatures/Source/GameFeatures/Private/GameFeatureData.cpp b/Engine/Plugins/Experimental/GameFeatures/Source/GameFeatures/Private/GameFeatureData.cpp
--- a/Engine/Plugins/Experimental/GameFeatures/Source/GameFeatures/Private/GameFeatureData.cpp
+++ b/Engine/Plugins/Experimental/GameFeatures/Source/GameFeatures/Private/GameFeatureData.cpp
@@ -106,18 +106,10 @@ void UGameFeatureData::InitializeHierarchicalPluginIniFiles(const FString& Plugi
 		FConfigFile TempConfig;
 		if (FConfigCacheIni::LoadExternalIniFile(TempConfig, *PluginIniName, *EngineConfigDir, *PluginConfigDir, bIsBaseIniName, nullptr, bForceReloadFromDisk, bWriteDestIni) && (TempConfig.Num() > 0))
 		{
-			FString IniFile;
-			if (IniName.Equals(TEXT("Input")))
-			{
-				IniFile = GInputIni;
-			}
-			else
-			{
-				IniFile = FString::Printf(TEXT("%s%s/%s.ini"), *FPaths::GeneratedConfigDir(), ANSI_TO_TCHAR(FPlatformProperties::PlatformName()), *IniName);
-			}
+			// Need to get the in-memory config filename, the on disk one is likely not up to date
+			FString IniFile = GConfig->GetConfigFilename(*IniName);
 
-			FPaths::MakeStandardFilename(IniFile);
-			if (FConfigFile* ExistingConfig = GConfig->Find(IniFile))
+			if (FConfigFile* ExistingConfig = GConfig->FindConfigFile(IniFile))
 			{
 				// @todo: Might want to consider modifying the engine level's API here to allow for a combination that yields affected
 				// sections and/or optionally just does the reload itself. This route is less efficient than it needs to be, resulting in parsing twice, 
diff --git a/Engine/Source/Runtime/Projects/Private/PluginManager.cpp b/Engine/Source/Runtime/Projects/Private/PluginManager.cpp
--- a/Engine/Source/Runtime/Projects/Private/PluginManager.cpp
+++ b/Engine/Source/Runtime/Projects/Private/PluginManager.cpp
@@ -625,11 +625,7 @@ bool FPluginManager::IntegratePluginsIntoConfig(FConfigCacheIni& ConfigSystem, c
 			IFileManager::Get().FindFiles(PluginConfigs, *PluginConfigDir, TEXT("ini"));
 			for (const FString& ConfigFile : PluginConfigs)
 			{
-				// Optionally strip Default and Base from source filename, old behavior would create a new misnamed file which allowed redirects to work
-				// But, it's better to just apply it to the config actually desired by the user
 				FString BaseConfigFile = *FPaths::GetBaseFilename(ConfigFile);
-				BaseConfigFile.RemoveFromStart(TEXT("Base"));
-				BaseConfigFile.RemoveFromStart(TEXT("Default"));
 
 				// Use GetConfigFilename to find the proper config file to combine into, since it manages command line overrides and path sanitization
 				FString PluginConfigFilename = ConfigSystem.GetConfigFilename(*BaseConfigFile);
@@ -1035,11 +1031,13 @@ bool FPluginManager::ConfigureEnabledPlugins()
 				IFileManager::Get().FindFiles(PluginConfigs, *PluginConfigDir, TEXT("ini"));
 				for (const FString& ConfigFile : PluginConfigs)
 				{
-					// Optionally strip Default and Base from source filename, old behavior would create a new misnamed file which allowed redirects to work
-					// But, it's better to just apply it to the config actually desired by the user
 					FString BaseConfigFile = *FPaths::GetBaseFilename(ConfigFile);					
-					BaseConfigFile.RemoveFromStart(TEXT("Base"));
-					BaseConfigFile.RemoveFromStart(TEXT("Default"));
+
+					if (BaseConfigFile == Plugin.Name)
+					{
+						// We just handles this, skip it
+						continue;
+					}
 
 					// Build the config system key for the overridden config
 					PluginConfigFilename = GConfig->GetConfigFilename(*BaseConfigFile);
