PR #8043: Fix DisableEnginePluginsByDefault inconsistency between ProjectDescriptor.cpp and ProjectDescriptor.cs (Contributed by [REDACTED_PERSON])

diff --git a/Engine/Source/Runtime/Projects/Private/ProjectDescriptor.cpp b/Engine/Source/Runtime/Projects/Private/ProjectDescriptor.cpp
--- a/Engine/Source/Runtime/Projects/Private/ProjectDescriptor.cpp
+++ b/Engine/Source/Runtime/Projects/Private/ProjectDescriptor.cpp
@@ -257,6 +257,11 @@ void FProjectDescriptor::Write(TJsonWriter<>& Writer, const FString& PathToProje
 	Writer.WriteValue(TEXT("Category"), Category);
 	Writer.WriteValue(TEXT("Description"), Description);
 
+	if (bDisableEnginePluginsByDefault)
+	{
+		Writer.WriteValue(TEXT("DisableEnginePluginsByDefault"), bDisableEnginePluginsByDefault);
+	}
+
 	// Write the enterprise flag
 	if (bIsEnterpriseProject)
 	{
