Don't change the plugin search path when mouting a plugin with an explicit uplugin file name (unless it's requested to update the search path of the project)
This removes a costly RefreshPluginList call and prevents poluting the plugin manager with plugins we want to forget about once they are unloaded

diff --git a/Engine/Plugins/Developer/PluginUtils/Source/PluginUtils/Private/PluginUtils.cpp b/Engine/Plugins/Developer/PluginUtils/Source/PluginUtils/Private/PluginUtils.cpp
--- a/Engine/Plugins/Developer/PluginUtils/Source/PluginUtils/Private/PluginUtils.cpp
+++ b/Engine/Plugins/Developer/PluginUtils/Source/PluginUtils/Private/PluginUtils.cpp
@@ -278,9 +278,19 @@ namespace PluginUtils
 	{
 		ensure(!PluginLocation.IsEmpty());
 
-		FPluginUtils::AddToPluginSearchPathIfNeeded(PluginLocation, /*bRefreshPlugins*/ false, MountParams.bUpdateProjectPluginSearchPath);
+		const FString PluginFilePath = FPluginUtils::GetPluginFilePath(PluginLocation, PluginName, /*bFullPath*/ true);
 
-		IPluginManager::Get().RefreshPluginsList();
+		if (MountParams.bUpdateProjectPluginSearchPath)
+		{
+			FPluginUtils::AddToPluginSearchPathIfNeeded(PluginLocation, /*bRefreshPlugins=*/true, /*bUpdateProjectFile=*/true);
+		}
+		else
+		{
+			if (!IPluginManager::Get().AddToPluginsList(PluginFilePath, &FailReason))
+			{
+				return nullptr;
+			}
+		}
 
 		// Find the plugin in the manager.
 		TSharedPtr<IPlugin> Plugin = IPluginManager::Get().FindPlugin(PluginName);
@@ -291,7 +301,6 @@ namespace PluginUtils
 		}
 
 		// Double check the path matches
-		const FString PluginFilePath = FPluginUtils::GetPluginFilePath(PluginLocation, PluginName, /*bFullPath*/ true);
 		if (!FPaths::IsSamePath(Plugin->GetDescriptorFileName(), PluginFilePath))
 		{
 			const FString PluginFilePathFull = FPaths::ConvertRelativePathToFull(Plugin->GetDescriptorFileName());
diff --git a/Engine/Plugins/Developer/PluginUtils/Source/PluginUtils/Public/PluginUtils.h b/Engine/Plugins/Developer/PluginUtils/Source/PluginUtils/Public/PluginUtils.h
--- a/Engine/Plugins/Developer/PluginUtils/Source/PluginUtils/Public/PluginUtils.h
+++ b/Engine/Plugins/Developer/PluginUtils/Source/PluginUtils/Public/PluginUtils.h
@@ -125,7 +125,6 @@ public:
 		/**
 		 * Whether to update the project additional plugin directories (persistently saved in uproject file)
 		 * if the plugin location is not under the engine or project plugin folder.
-		 * Otherwise the plugin search path gets updated for the process lifetime only.
 		 */
 		bool bUpdateProjectPluginSearchPath = true;
 
diff --git a/Engine/Source/Runtime/Projects/Private/PluginManager.cpp b/Engine/Source/Runtime/Projects/Private/PluginManager.cpp
--- a/Engine/Source/Runtime/Projects/Private/PluginManager.cpp
+++ b/Engine/Source/Runtime/Projects/Private/PluginManager.cpp
@@ -219,7 +219,7 @@ void FPluginManager::RefreshPluginsList()
 	}
 }
 
-bool FPluginManager::AddToPluginsList(const FString& PluginFilename)
+bool FPluginManager::AddToPluginsList(const FString& PluginFilename, FText* OutFailReason /*= nullptr*/)
 {
 #if (WITH_ENGINE && !IS_PROGRAM) || WITH_PLUGIN_SUPPORT
 	// No need to readd if it already exists
@@ -231,8 +231,8 @@ bool FPluginManager::AddToPluginsList(const FString& PluginFilename)
 
 	// Read the plugin and load it
 	FPluginDescriptor Descriptor;
-	FText FailureReason;
-	if (Descriptor.Load(PluginFilename, FailureReason))
+	FText FailReason;
+	if (Descriptor.Load(PluginFilename, FailReason))
 	{
 		// Determine the plugin type
 		EPluginType PluginType = EPluginType::External;
@@ -271,7 +271,12 @@ bool FPluginManager::AddToPluginsList(const FString& PluginFilename)
 	}
 	else
 	{
-		UE_LOG(LogPluginManager, Warning, TEXT("AddToPluginsList failed to load plugin %s. Reason: %s"), *PluginFilename, *FailureReason.ToString());
+		FailReason = FText::Format(LOCTEXT("AddToPluginsListFailed", "Failed to load plugin '{0}'\n{1}"), FText::FromString(PluginFilename), FailReason);
+		UE_LOG(LogPluginManager, Error, TEXT("%s"), *FailReason.ToString());
+		if (OutFailReason)
+		{
+			*OutFailReason = MoveTemp(FailReason);
+		}
 	}
 #endif
 
diff --git a/Engine/Source/Runtime/Projects/Private/PluginManager.h b/Engine/Source/Runtime/Projects/Private/PluginManager.h
--- a/Engine/Source/Runtime/Projects/Private/PluginManager.h
+++ b/Engine/Source/Runtime/Projects/Private/PluginManager.h
@@ -105,7 +105,7 @@ public:
 
 	/** IPluginManager interface */
 	virtual void RefreshPluginsList() override;
-	virtual bool AddToPluginsList( const FString& PluginFilename ) override;
+	virtual bool AddToPluginsList(const FString& PluginFilename, FText* OutFailReason = nullptr) override;
 	virtual bool LoadModulesForEnabledPlugins( const ELoadingPhase::Type LoadingPhase ) override;
 	virtual FLoadingModulesForPhaseEvent& OnLoadingPhaseComplete() override;
 	virtual ELoadingPhase::Type GetLastCompletedLoadingPhase() const override;
diff --git a/Engine/Source/Runtime/Projects/Public/Interfaces/IPluginManager.h b/Engine/Source/Runtime/Projects/Public/Interfaces/IPluginManager.h
--- a/Engine/Source/Runtime/Projects/Public/Interfaces/IPluginManager.h
+++ b/Engine/Source/Runtime/Projects/Public/Interfaces/IPluginManager.h
@@ -199,7 +199,7 @@ public:
 	 * 
 	 * @return True if the plugin was added or already in the list. False if it failed to load.
 	 */
-	virtual bool AddToPluginsList( const FString& PluginFilename ) = 0;
+	virtual bool AddToPluginsList(const FString& PluginFilename, FText* OutFailReason = nullptr) = 0;
 
 	/**
 	 * Loads all plug-ins
