Precache a list of all available INI files for the project when mounting plugins.
This reduces the number of system calls to determine if INI files exist, which significantly reduces mounting time for plugins.

#changelist validated

diff --git a/Engine/Source/Runtime/Core/Private/Misc/ConfigCacheIni.cpp b/Engine/Source/Runtime/Core/Private/Misc/ConfigCacheIni.cpp
--- a/Engine/Source/Runtime/Core/Private/Misc/ConfigCacheIni.cpp
+++ b/Engine/Source/Runtime/Core/Private/Misc/ConfigCacheIni.cpp
@@ -41,6 +41,8 @@
 #define DISABLE_GENERATED_INI_WHEN_COOKED 0
 #endif
 
+const TSet<FString>* FConfigCacheIni::IniCacheSet = nullptr;
+
 namespace
 {
 	static FName VersionName("Version");
@@ -474,8 +476,38 @@ static bool DoesConfigFileExistWrapper(const TCHAR* IniFile)
 		return true;
 	}
 
+	// So far, testing on cooked consoles, cooked desktop, and the Editor
+	// works fine.
+	// However, there was an issue where INIs wouldn't be found during cooking,
+	// which would pass by silently. Realistically, in most cases this would never
+	// have caused an issue, but using FPlatformProperties::RequiresCookedData
+	// to prevent using the cache in that case ensures full consistency.
+	const TSet<FString>* const IniCacheSet = FConfigCacheIni::GetIniCacheSet();
+	if (IniCacheSet && FPlatformProperties::RequiresCookedData())
+	{ 
+		const FString IniFileString(IniFile);
+		const bool bFileExistsCached = IniCacheSet->Contains(IniFileString);
+
+		// This code can be uncommented if we expect there are INIs that are not being
+		// found in the cache.
+		/**
+		const bool bFileExistsCachedTest = IFileManager::Get().FileSize(IniFile) >= 0;
+		ensureMsgf(
+			bFileExistsCached == bFileExistsCachedTest,
+			TEXT("DoesConfigFileExistWrapper: InCache = %d, InFileSystem = %d, Name = %s, Configs = \n%s"),
+			!!bFileExistsCached,
+			!!bFileExistsCachedTest,
+			IniFile,
+			*FString::Join(*IniCacheSet, TEXT("\n"))
+		);
+		*/
+		
+		return bFileExistsCached;
+	}
+
 	// otherwise just look for the normal file to exist
-	return IFileManager::Get().FileSize(IniFile) >= 0;
+	const bool bFileExistsCached = IFileManager::Get().FileSize(IniFile) >= 0;
+	return bFileExistsCached;
 }
 
 /**
diff --git a/Engine/Source/Runtime/Core/Public/Misc/ConfigCacheIni.h b/Engine/Source/Runtime/Core/Public/Misc/ConfigCacheIni.h
--- a/Engine/Source/Runtime/Core/Public/Misc/ConfigCacheIni.h
+++ b/Engine/Source/Runtime/Core/Public/Misc/ConfigCacheIni.h
@@ -659,6 +659,23 @@ enum class EConfigCacheType : uint8
 class CORE_API FConfigCacheIni : private TMap<FString,FConfigFile>
 {
 public:
+
+	static void SetIniCacheSet(const TSet<FString>* InIniCacheSet)
+	{
+		IniCacheSet = InIniCacheSet;
+	}
+
+	static const TSet<FString>* GetIniCacheSet()
+	{
+		return IniCacheSet;
+	}
+
+private:
+
+	static const TSet<FString>* IniCacheSet;
+	
+public:
+	
 	// Basic functions.
 	FConfigCacheIni(EConfigCacheType Type);
 
diff --git a/Engine/Source/Runtime/Projects/Private/PluginManager.cpp b/Engine/Source/Runtime/Projects/Private/PluginManager.cpp
--- a/Engine/Source/Runtime/Projects/Private/PluginManager.cpp
+++ b/Engine/Source/Runtime/Projects/Private/PluginManager.cpp
@@ -74,6 +74,16 @@ namespace PluginSystemDefs
 
 		return PluginCount;
 	}
+
+	// TODO: This processing happens too early to actually get picked up by ExecCmds,
+	// so it has to be reset in code when A/B testing. May be fine for INI settings though,
+	// otherwise, we might need a command line setting instead?
+	static bool GbCacheIniFilesForProcessing = true;
+	static FAutoConsoleVariableRef CVarCacheIniFilesForProcessing(
+		TEXT("c.CacheIniFilesForProcessing"),
+		GbCacheIniFilesForProcessing,
+		TEXT("When enabled, we'll scrape top level known directories for INI files and use that instead of system calls when processing plugins. Can decrease startup times.")
+	);
 }
 
 FPlugin::FPlugin(const FString& InFileName, const FPluginDescriptor& InDescriptor, EPluginType InType)
@@ -992,10 +1002,28 @@ bool FPluginManager::ConfigureEnabledPlugins()
 				}
 			}
 
+			TSet<FString> AllIniFiles;
+
+			if (PluginSystemDefs::GbCacheIniFilesForProcessing)
+			{
+				SCOPED_BOOT_TIMING("ParallelPluginEnabling::FindIniFiles");
+				TArray<FString> AllIniFilesList;
+				
+				// Using ProjectDir and EngineDir are really broad, but they also cover all plugin config directories individually
+				// and should cover all of the extension directories as well.
+				IFileManager::Get().FindFilesRecursive(AllIniFilesList, *FPaths::EngineDir(), TEXT("*.ini"), /*Files=*/true, /*Directories=*/false, /*bClearFileNames=*/false);
+				IFileManager::Get().FindFilesRecursive(AllIniFilesList, *FPaths::ProjectDir(), TEXT("*.ini"), /*Files=*/true, /*Directories=*/false, /*bClearFileNames=*/false);
+				IFileManager::Get().FindFilesRecursive(AllIniFilesList, *FPaths::ProjectSavedDir(), TEXT("*.ini"), /*Files=*/true, /*Directories=*/false, /*bClearFileNames=*/false);
+
+				AllIniFiles.Append(AllIniFilesList);
+			}
+
+			
 			FCriticalSection ConfigCS;
 			FCriticalSection PluginPakCS;
+			
 			// Mount all the enabled plugins
-			ParallelFor(PluginsArray.Num(), [&PluginsArray, &ConfigCS, &PluginPakCS, this](int32 Index)
+			ParallelFor(PluginsArray.Num(), [&PluginsArray, &ConfigCS, &PluginPakCS, &AllIniFiles, this](int32 Index)
 			{
 				FString PlatformName = FPlatformProperties::PlatformName();
 				FPlugin& Plugin = *PluginsArray[Index];
@@ -1019,16 +1047,30 @@ bool FPluginManager::ConfigureEnabledPlugins()
 				// Build the config system key for PluginName.ini
 				FString PluginConfigFilename = GConfig->GetConfigFilename(*Plugin.Name);
 				{
+					
 					FScopeLock Locker(&ConfigCS);
 
 					FConfigFile& PluginConfig = GConfig->Add(PluginConfigFilename, FConfigFile());
 
+					// We probably *don't* need to set and unset this every time,
+					// but this feels safer in case INI files are added for a plugin and then we try to load it later.
+					// Dynamically added / generated INI files should be irrelevant for the LoadExternalIniFile call.
+					if (PluginSystemDefs::GbCacheIniFilesForProcessing)
+					{
+						FConfigCacheIni::SetIniCacheSet(&AllIniFiles);
+					}	
+
 					// This will write out an ini to PluginConfigFilename
 					if (!FConfigCacheIni::LoadExternalIniFile(PluginConfig, *Plugin.Name, *EngineConfigDir, *SourceConfigDir, true, nullptr, false, true))
 					{
 						// Nothing to add, remove from map
 						GConfig->Remove(PluginConfigFilename);
 					}
+
+					if (PluginSystemDefs::GbCacheIniFilesForProcessing)
+					{
+						FConfigCacheIni::SetIniCacheSet(nullptr);
+					}	
 				}
 
 				// override config cache entries with plugin configs (Engine.ini, Game.ini, etc in <PluginDir>\Config\)
