Works arround 3min long iteration time regression on -run=CookGlobalShaders -platform=Windows -stage -SkipAssetScan

Change 17815600 causes to skan the project's asset directory early at launch of the editor

diff --git a/Engine/Source/Runtime/Projects/Private/PluginManager.cpp b/Engine/Source/Runtime/Projects/Private/PluginManager.cpp
--- a/Engine/Source/Runtime/Projects/Private/PluginManager.cpp
+++ b/Engine/Source/Runtime/Projects/Private/PluginManager.cpp
@@ -84,6 +84,17 @@ namespace PluginSystemDefs
 		GbCacheIniFilesForProcessing,
 		TEXT("When enabled, we'll scrape top level known directories for INI files and use that instead of system calls when processing plugins. Can decrease startup times.")
 	);
+
+	bool IsCachingIniFilesForProcessing()
+	{
+#if WITH_EDITOR
+		if (FParse::Param(FCommandLine::Get(), TEXT("SkipAssetScan")))
+		{
+			return false;
+		}
+#endif
+		return GbCacheIniFilesForProcessing;
+	}
 }
 
 FPlugin::FPlugin(const FString& InFileName, const FPluginDescriptor& InDescriptor, EPluginType InType)
@@ -1004,7 +1015,7 @@ bool FPluginManager::ConfigureEnabledPlugins()
 
 			TSet<FString> AllIniFiles;
 
-			if (PluginSystemDefs::GbCacheIniFilesForProcessing)
+			if (PluginSystemDefs::IsCachingIniFilesForProcessing())
 			{
 				SCOPED_BOOT_TIMING("ParallelPluginEnabling::FindIniFiles");
 				TArray<FString> AllIniFilesList;
@@ -1055,7 +1066,7 @@ bool FPluginManager::ConfigureEnabledPlugins()
 					// We probably *don't* need to set and unset this every time,
 					// but this feels safer in case INI files are added for a plugin and then we try to load it later.
 					// Dynamically added / generated INI files should be irrelevant for the LoadExternalIniFile call.
-					if (PluginSystemDefs::GbCacheIniFilesForProcessing)
+					if (PluginSystemDefs::IsCachingIniFilesForProcessing())
 					{
 						FConfigCacheIni::SetIniCacheSet(&AllIniFiles);
 					}	
@@ -1067,7 +1078,7 @@ bool FPluginManager::ConfigureEnabledPlugins()
 						GConfig->Remove(PluginConfigFilename);
 					}
 
-					if (PluginSystemDefs::GbCacheIniFilesForProcessing)
+					if (PluginSystemDefs::IsCachingIniFilesForProcessing())
 					{
 						FConfigCacheIni::SetIniCacheSet(nullptr);
 					}	
