Fixed crash during cook on a relative path

#rnx

diff --git a/Engine/Source/Runtime/Projects/Private/ProjectDescriptor.cpp b/Engine/Source/Runtime/Projects/Private/ProjectDescriptor.cpp
--- a/Engine/Source/Runtime/Projects/Private/ProjectDescriptor.cpp
+++ b/Engine/Source/Runtime/Projects/Private/ProjectDescriptor.cpp
@@ -144,8 +144,11 @@ bool FProjectDescriptor::Read(const FJsonObject& Object, const FString& PathToPr
 		if (FPlatformProperties::RequiresCookedData() && AdditionalPluginDirectoriesValue->Num() > 0)
 		{
 			AdditionalPluginDirectories.Empty();
-			FString RemappedDir = FPaths::ProjectDir() + TEXT("../RemappedPlugins/");
-			RemappedDir = FPaths::ConvertRelativePathToFull(RemappedDir);
+			FString RemappedDir = FPaths::ProjectDir() / TEXT("../RemappedPlugins/");
+			if (FPaths::IsRelative(RemappedDir))
+			{
+				RemappedDir = IFileManager::Get().ConvertToAbsolutePathForExternalAppForWrite(*RemappedDir);
+			}
 			AddPluginDirectory(RemappedDir);
 		}
 	}
