New "CanContainVerse" plugin descriptor property
- allows to specifically filter plugins at runtime
- added checks in UBT that point out the incorrect presence/absence of the property

#rnx

diff --git a/Engine/Source/Programs/UnrealBuildTool/Configuration/UEBuildTarget.cs b/Engine/Source/Programs/UnrealBuildTool/Configuration/UEBuildTarget.cs
--- a/Engine/Source/Programs/UnrealBuildTool/Configuration/UEBuildTarget.cs
+++ b/Engine/Source/Programs/UnrealBuildTool/Configuration/UEBuildTarget.cs
@@ -3454,8 +3454,8 @@ namespace UnrealBuildTool
 				}
 				Instance.Dependencies = Dependencies;
 
-				// Stage the descriptor if the plugin contains content
-				if (Info.Descriptor.bCanContainContent || Dependencies.Any(x => x.bDescriptorNeededAtRuntime))
+				// Stage the descriptor if the plugin contains content or Verse code
+				if (Info.Descriptor.bCanContainContent || Info.Descriptor.bCanContainVerse || Dependencies.Any(x => x.bDescriptorNeededAtRuntime))
 				{
 					Instance.bDescriptorNeededAtRuntime = true;
 				}
diff --git a/Engine/Source/Programs/UnrealBuildTool/System/PluginDescriptor.cs b/Engine/Source/Programs/UnrealBuildTool/System/PluginDescriptor.cs
--- a/Engine/Source/Programs/UnrealBuildTool/System/PluginDescriptor.cs
+++ b/Engine/Source/Programs/UnrealBuildTool/System/PluginDescriptor.cs
@@ -149,6 +149,11 @@ namespace UnrealBuildTool
 		/// </summary>
 		public bool bCanContainContent;
 
+		/// <summary>
+		/// Can this plugin contain Verse code (either in content directory or in any of its modules)?
+		/// </summary>
+		public bool bCanContainVerse;
+
 		/// <summary>
 		/// Marks the plugin as beta in the UI
 		/// </summary>
@@ -290,6 +295,7 @@ namespace UnrealBuildTool
 			}
 
 			RawObject.TryGetBoolField("CanContainContent", out bCanContainContent);
+			RawObject.TryGetBoolField("CanContainVerse", out bCanContainVerse);
 			RawObject.TryGetBoolField("IsBetaVersion", out bIsBetaVersion);
 			RawObject.TryGetBoolField("IsExperimentalVersion", out bIsExperimentalVersion);
 			RawObject.TryGetBoolField("Installed", out bInstalled);
@@ -381,6 +387,7 @@ namespace UnrealBuildTool
 				Writer.WriteValue("EnabledByDefault", bEnabledByDefault.Value);
 			}
 			Writer.WriteValue("CanContainContent", bCanContainContent);
+			Writer.WriteValue("CanContainVerse", bCanContainVerse);
 			if (bIsBetaVersion)
 			{
 				Writer.WriteValue("IsBetaVersion", bIsBetaVersion);
diff --git a/Engine/Source/Runtime/Projects/Private/PluginDescriptor.cpp b/Engine/Source/Runtime/Projects/Private/PluginDescriptor.cpp
--- a/Engine/Source/Runtime/Projects/Private/PluginDescriptor.cpp
+++ b/Engine/Source/Runtime/Projects/Private/PluginDescriptor.cpp
@@ -79,6 +79,7 @@ FPluginDescriptor::FPluginDescriptor()
 	: Version(0)
 	, EnabledByDefault(EPluginEnabledByDefault::Unspecified)
 	, bCanContainContent(false)
+	, bCanContainVerse(false)
 	, bIsBetaVersion(false)
 	, bIsExperimentalVersion(false)
 	, bInstalled(false)
@@ -216,6 +217,7 @@ bool FPluginDescriptor::Read(const FJsonObject& Object, FText* OutFailReason /*=
 	}
 
 	Object.TryGetBoolField(TEXT("CanContainContent"), bCanContainContent);
+	Object.TryGetBoolField(TEXT("CanContainVerse"), bCanContainVerse);
 	Object.TryGetBoolField(TEXT("IsBetaVersion"), bIsBetaVersion);
 	Object.TryGetBoolField(TEXT("IsExperimentalVersion"), bIsExperimentalVersion);
 	Object.TryGetBoolField(TEXT("Installed"), bInstalled);
@@ -328,6 +330,7 @@ void FPluginDescriptor::UpdateJson(FJsonObject& JsonObject) const
 	}
 
 	JsonObject.SetBoolField(TEXT("CanContainContent"), bCanContainContent);
+	JsonObject.SetBoolField(TEXT("CanContainVerse"), bCanContainVerse);
 	JsonObject.SetBoolField(TEXT("IsBetaVersion"), bIsBetaVersion);
 	JsonObject.SetBoolField(TEXT("IsExperimentalVersion"), bIsExperimentalVersion);
 	JsonObject.SetBoolField(TEXT("Installed"), bInstalled);
diff --git a/Engine/Source/Runtime/Projects/Private/PluginManager.cpp b/Engine/Source/Runtime/Projects/Private/PluginManager.cpp
--- a/Engine/Source/Runtime/Projects/Private/PluginManager.cpp
+++ b/Engine/Source/Runtime/Projects/Private/PluginManager.cpp
@@ -1160,7 +1160,7 @@ bool FPluginManager::ConfigureEnabledPlugins()
 			}, true); // @todo disable parallelism for now as it's causing hard to track problems
 		}
 
-		for (TSharedRef<IPlugin> Plugin: GetEnabledPluginsWithContent())
+		for (TSharedRef<IPlugin> Plugin: GetEnabledPluginsWithContentOrVerse())
 		{
 			if (Plugin->GetDescriptor().bExplicitlyLoaded)
 			{
@@ -1790,6 +1790,7 @@ TSharedPtr<IPlugin> FPluginManager::FindPluginFromPath(const FString& PluginPath
 TArray<TSharedRef<IPlugin>> FPluginManager::GetEnabledPlugins()
 {
 	TArray<TSharedRef<IPlugin>> Plugins;
+	Plugins.Reserve(AllPlugins.Num());
 	for(TPair<FString, TSharedRef<FPlugin>>& PluginPair : AllPlugins)
 	{
 		TSharedRef<FPlugin>& PossiblePlugin = PluginPair.Value;
@@ -1804,6 +1805,7 @@ TArray<TSharedRef<IPlugin>> FPluginManager::GetEnabledPlugins()
 TArray<TSharedRef<IPlugin>> FPluginManager::GetEnabledPluginsWithContent() const
 {
 	TArray<TSharedRef<IPlugin>> Plugins;
+	Plugins.Reserve(AllPlugins.Num());
 	for (const TPair<FString, TSharedRef<FPlugin>>& PluginPair : AllPlugins)
 	{
 		const TSharedRef<FPlugin>& PluginRef = PluginPair.Value;
@@ -1816,6 +1818,38 @@ TArray<TSharedRef<IPlugin>> FPluginManager::GetEnabledPluginsWithContent() const
 	return Plugins;
 }
 
+TArray<TSharedRef<IPlugin>> FPluginManager::GetEnabledPluginsWithVerse() const
+{
+	TArray<TSharedRef<IPlugin>> Plugins;
+	Plugins.Reserve(AllPlugins.Num());
+	for (const TPair<FString, TSharedRef<FPlugin>>& PluginPair : AllPlugins)
+	{
+		const TSharedRef<FPlugin>& PluginRef = PluginPair.Value;
+		const FPlugin& Plugin = *PluginRef;
+		if (Plugin.IsEnabled() && Plugin.CanContainVerse())
+		{
+			Plugins.Add(PluginRef);
+		}
+	}
+	return Plugins;
+}
+
+TArray<TSharedRef<IPlugin>> FPluginManager::GetEnabledPluginsWithContentOrVerse() const
+{
+	TArray<TSharedRef<IPlugin>> Plugins;
+	Plugins.Reserve(AllPlugins.Num());
+	for (const TPair<FString, TSharedRef<FPlugin>>& PluginPair : AllPlugins)
+	{
+		const TSharedRef<FPlugin>& PluginRef = PluginPair.Value;
+		const FPlugin& Plugin = *PluginRef;
+		if (Plugin.IsEnabled() && (Plugin.CanContainContent() || Plugin.CanContainVerse()))
+		{
+			Plugins.Add(PluginRef);
+		}
+	}
+	return Plugins;
+}
+
 TArray<TSharedRef<IPlugin>> FPluginManager::GetDiscoveredPlugins()
 {
 	TArray<TSharedRef<IPlugin>> Plugins;
@@ -1830,6 +1864,7 @@ TArray<TSharedRef<IPlugin>> FPluginManager::GetDiscoveredPlugins()
 TArray< FPluginStatus > FPluginManager::QueryStatusForAllPlugins() const
 {
 	TArray< FPluginStatus > PluginStatuses;
+	PluginStatuses.Reserve(AllPlugins.Num());
 
 	for( const TPair<FString, TSharedRef<FPlugin>>& PluginPair : AllPlugins )
 	{
@@ -1919,7 +1954,7 @@ void FPluginManager::MountPluginFromExternalSource(const TSharedRef<FPlugin>& Pl
 	Plugin->bEnabled = true;
 
 	// Mount the plugin content directory
-	if (Plugin->CanContainContent() && ensure(RegisterMountPointDelegate.IsBound()))
+	if ((Plugin->CanContainContent() || Plugin->CanContainVerse()) && ensure(RegisterMountPointDelegate.IsBound()))
 	{
 		if (NewPluginMountedEvent.IsBound())
 		{
@@ -1930,16 +1965,19 @@ void FPluginManager::MountPluginFromExternalSource(const TSharedRef<FPlugin>& Pl
 		RegisterMountPointDelegate.Execute(Plugin->GetMountedAssetPath(), ContentDir);
 
 		// Register this plugin's path with the list of content directories that the editor will search
-		if (FConfigFile* EngineConfigFile = GConfig->Find(GEngineIni))
+		if (Plugin->CanContainContent())
 		{
-			if (FConfigSection* CoreSystemSection = EngineConfigFile->Find(TEXT("Core.System")))
+			if (FConfigFile* EngineConfigFile = GConfig->Find(GEngineIni))
 			{
-				CoreSystemSection->AddUnique("Paths", MoveTemp(ContentDir));
+				if (FConfigSection* CoreSystemSection = EngineConfigFile->Find(TEXT("Core.System")))
+				{
+					CoreSystemSection->AddUnique("Paths", MoveTemp(ContentDir));
+				}
 			}
-		}
 
-		// Update the localization cache for the newly added content directory
-		UpdatePackageLocalizationCacheDelegate.ExecuteIfBound();
+			// Update the localization cache for the newly added content directory
+			UpdatePackageLocalizationCacheDelegate.ExecuteIfBound();
+		}
 	}
 
 	// If it's a code module, also load the modules for it
@@ -1997,7 +2035,7 @@ bool FPluginManager::UnmountPluginFromExternalSource(const TSharedPtr<FPlugin>&
 		return false;
 	}
 
-	if (Plugin->CanContainContent() && ensure(UnRegisterMountPointDelegate.IsBound()))
+	if ((Plugin->CanContainContent() || Plugin->CanContainVerse()) && ensure(UnRegisterMountPointDelegate.IsBound()))
 	{
 		UnRegisterMountPointDelegate.Execute(Plugin->GetMountedAssetPath(), Plugin->GetContentDir());
 	}
diff --git a/Engine/Source/Runtime/Projects/Private/PluginManager.h b/Engine/Source/Runtime/Projects/Private/PluginManager.h
--- a/Engine/Source/Runtime/Projects/Private/PluginManager.h
+++ b/Engine/Source/Runtime/Projects/Private/PluginManager.h
@@ -77,6 +77,11 @@ public:
 		return Descriptor.bCanContainContent;
 	}
 
+	virtual bool CanContainVerse() const override
+	{
+		return Descriptor.bCanContainVerse;
+	}
+
 	virtual EPluginType GetType() const override
 	{
 		return Type;
@@ -129,6 +134,8 @@ public:
 
 	virtual TArray<TSharedRef<IPlugin>> GetEnabledPlugins() override;
 	virtual TArray<TSharedRef<IPlugin>> GetEnabledPluginsWithContent() const override;
+	virtual TArray<TSharedRef<IPlugin>> GetEnabledPluginsWithVerse() const override;
+	virtual TArray<TSharedRef<IPlugin>> GetEnabledPluginsWithContentOrVerse() const override;
 	virtual TArray<TSharedRef<IPlugin>> GetDiscoveredPlugins() override;
 	virtual TArray< FPluginStatus > QueryStatusForAllPlugins() const override;
 	virtual bool AddPluginSearchPath(const FString& ExtraDiscoveryPath, bool bRefresh = true) override;
diff --git a/Engine/Source/Runtime/Projects/Public/Interfaces/IPluginManager.h b/Engine/Source/Runtime/Projects/Public/Interfaces/IPluginManager.h
--- a/Engine/Source/Runtime/Projects/Public/Interfaces/IPluginManager.h
+++ b/Engine/Source/Runtime/Projects/Public/Interfaces/IPluginManager.h
@@ -148,6 +148,13 @@ public:
 	 */
 	virtual bool CanContainContent() const = 0;
 
+	/**
+	 * Determines if the plugin can contain Verse code.
+	 *
+	 * @return True if the plugin can contain Verse code.
+	 */
+	virtual bool CanContainVerse() const = 0;
+
 	/**
 	 * Returns the plugin's location
 	 *
@@ -297,6 +304,18 @@ public:
 	 */
 	virtual TArray<TSharedRef<IPlugin>> GetEnabledPluginsWithContent() const = 0;
 
+	/**
+	 * Gets an array of all enabled plugins that can have Verse code.
+	 *
+	 * @return	Array of plugins with IsEnabled() and CanContainVerse() both true.
+	 */
+	virtual TArray<TSharedRef<IPlugin>> GetEnabledPluginsWithVerse() const = 0;
+
+	/**
+	 * Gets an array of all enabled plugins that can have content or Verse code.
+	 */
+	virtual TArray<TSharedRef<IPlugin>> GetEnabledPluginsWithContentOrVerse() const = 0;
+
 	/**
 	 * Gets an array of all the discovered plugins.
 	 *
diff --git a/Engine/Source/Runtime/Projects/Public/PluginDescriptor.h b/Engine/Source/Runtime/Projects/Public/PluginDescriptor.h
--- a/Engine/Source/Runtime/Projects/Public/PluginDescriptor.h
+++ b/Engine/Source/Runtime/Projects/Public/PluginDescriptor.h
@@ -87,6 +87,9 @@ struct PROJECTS_API FPluginDescriptor
 	/** Can this plugin contain content? */
 	bool bCanContainContent;
 
+	/** Can this plugin contain Verse code? */
+	bool bCanContainVerse;
+
 	/** Marks the plugin as beta in the UI */
 	bool bIsBetaVersion;
 
