Do not write the CanContainVerse property for plugin descriptors when it is false

#rnx

diff --git a/Engine/Source/Programs/UnrealBuildTool/System/PluginDescriptor.cs b/Engine/Source/Programs/UnrealBuildTool/System/PluginDescriptor.cs
--- a/Engine/Source/Programs/UnrealBuildTool/System/PluginDescriptor.cs
+++ b/Engine/Source/Programs/UnrealBuildTool/System/PluginDescriptor.cs
@@ -387,7 +387,10 @@ namespace UnrealBuildTool
 				Writer.WriteValue("EnabledByDefault", bEnabledByDefault.Value);
 			}
 			Writer.WriteValue("CanContainContent", bCanContainContent);
-			Writer.WriteValue("CanContainVerse", bCanContainVerse);
+			if (bCanContainVerse)
+			{
+				Writer.WriteValue("CanContainVerse", bCanContainVerse);
+			}
 			if (bIsBetaVersion)
 			{
 				Writer.WriteValue("IsBetaVersion", bIsBetaVersion);
diff --git a/Engine/Source/Runtime/Projects/Private/PluginDescriptor.cpp b/Engine/Source/Runtime/Projects/Private/PluginDescriptor.cpp
--- a/Engine/Source/Runtime/Projects/Private/PluginDescriptor.cpp
+++ b/Engine/Source/Runtime/Projects/Private/PluginDescriptor.cpp
@@ -330,7 +330,10 @@ void FPluginDescriptor::UpdateJson(FJsonObject& JsonObject) const
 	}
 
 	JsonObject.SetBoolField(TEXT("CanContainContent"), bCanContainContent);
-	JsonObject.SetBoolField(TEXT("CanContainVerse"), bCanContainVerse);
+	if (bCanContainVerse)
+	{
+		JsonObject.SetBoolField(TEXT("CanContainVerse"), bCanContainVerse);
+	}
 	JsonObject.SetBoolField(TEXT("IsBetaVersion"), bIsBetaVersion);
 	JsonObject.SetBoolField(TEXT("IsExperimentalVersion"), bIsExperimentalVersion);
 	JsonObject.SetBoolField(TEXT("Installed"), bInstalled);
