Fix editor loadtime regression from scanning for INIs.
This code path was intended primarily for Cooked / Packaged builds, where precaching all the INI files available in cooked files gave a significant reduction in startup times.
However, for editor builds, this actually causes significantly higher startup times (20s longer, or more).

diff --git a/Engine/Source/Runtime/Projects/Private/PluginManager.cpp b/Engine/Source/Runtime/Projects/Private/PluginManager.cpp
--- a/Engine/Source/Runtime/Projects/Private/PluginManager.cpp
+++ b/Engine/Source/Runtime/Projects/Private/PluginManager.cpp
@@ -75,25 +75,13 @@ namespace PluginSystemDefs
 		return PluginCount;
 	}
 
-	// TODO: This processing happens too early to actually get picked up by ExecCmds,
-	// so it has to be reset in code when A/B testing. May be fine for INI settings though,
-	// otherwise, we might need a command line setting instead?
-	static bool GbCacheIniFilesForProcessing = true;
-	static FAutoConsoleVariableRef CVarCacheIniFilesForProcessing(
-		TEXT("c.CacheIniFilesForProcessing"),
-		GbCacheIniFilesForProcessing,
-		TEXT("When enabled, we'll scrape top level known directories for INI files and use that instead of system calls when processing plugins. Can decrease startup times.")
-	);
 
 	bool IsCachingIniFilesForProcessing()
 	{
 #if WITH_EDITOR
-		if (FParse::Param(FCommandLine::Get(), TEXT("SkipAssetScan")))
-		{
-			return false;
-		}
+		return false;
 #endif
-		return GbCacheIniFilesForProcessing;
+		return true;
 	}
 }
 
