Allow stand-alone programs to opt out of plugin discoverability

#rnx

- By setting "UE_DISABLE_PLUGIN_DISCOVERY=1" in its target.cs a stand-alone program can avoid the cost of iterating over all plugins but still load and enable specific plugins that it knows it needs.
- In the future we could make this data driven by allowing a program to supply something similar to the .upluginmanifest file.

diff --git a/Engine/Source/Runtime/Projects/Private/PluginManager.cpp b/Engine/Source/Runtime/Projects/Private/PluginManager.cpp
--- a/Engine/Source/Runtime/Projects/Private/PluginManager.cpp
+++ b/Engine/Source/Runtime/Projects/Private/PluginManager.cpp
@@ -30,6 +30,9 @@ DEFINE_LOG_CATEGORY_STATIC( LogPluginManager, Log, All );
 
 #define LOCTEXT_NAMESPACE "PluginManager"
 
+#ifndef UE_DISABLE_PLUGIN_DISCOVERY
+	#define UE_DISABLE_PLUGIN_DISCOVERY 0
+#endif
 
 namespace PluginSystemDefs
 {
@@ -254,7 +257,7 @@ void FPluginManager::RefreshPluginsList()
 bool FPluginManager::AddToPluginsList(const FString& PluginFilename, FText* OutFailReason /*= nullptr*/)
 {
 #if (WITH_ENGINE && !IS_PROGRAM) || WITH_PLUGIN_SUPPORT
-	// No need to readd if it already exists
+	// No need to re-add if it already exists
 	FString PluginName = FPaths::GetBaseFilename(PluginFilename);
 	if (AllPlugins.Contains(PluginName))
 	{
@@ -331,7 +334,7 @@ void FPluginManager::DiscoverAllPlugins()
 
 void FPluginManager::ReadAllPlugins(TMap<FString, TSharedRef<FPlugin>>& Plugins, const TSet<FString>& ExtraSearchPaths)
 {
-#if (WITH_ENGINE && !IS_PROGRAM) || WITH_PLUGIN_SUPPORT
+#if (WITH_ENGINE && !IS_PROGRAM) || (WITH_PLUGIN_SUPPORT && !UE_DISABLE_PLUGIN_DISCOVERY)
 	const FProjectDescriptor* Project = IProjectManager::Get().GetCurrentProject();
 
 	// Find any plugin manifest files. These give us the plugin list (and their descriptors) without needing to scour the directory tree.
