Show dialog and block plugins from overriding specific config ini files.

Config files are marked NoSave when plugin overrides them.
NoSave prevents issues with plugins leaving stale data around in config if turned off later.

diff --git a/Engine/Config/BaseEngine.ini b/Engine/Config/BaseEngine.ini
--- a/Engine/Config/BaseEngine.ini
+++ b/Engine/Config/BaseEngine.ini
@@ -3006,6 +3006,10 @@ EditorUsesSpeed=FinalIfAvailable
 ; reuse as opposed to heap allocating every decode.
 PreallocatedBufferCount=2
 
+[Plugins]
+; Do not allow plugins to override EditorPerProjectUserSettings because it needs allow save
++ConfigFilesPluginsCannotOverride=EditorPerProjectUserSettings
+
 [ClothSettings]
 ;  Extra padding added to the bounds so that the collision can still be detected after being resolved
 EnvironmentCollisionPadding = 2.f
diff --git a/Engine/Source/Runtime/Projects/Private/PluginManager.cpp b/Engine/Source/Runtime/Projects/Private/PluginManager.cpp
--- a/Engine/Source/Runtime/Projects/Private/PluginManager.cpp
+++ b/Engine/Source/Runtime/Projects/Private/PluginManager.cpp
@@ -1029,6 +1029,9 @@ bool FPluginManager::ConfigureEnabledPlugins()
 				}
 			}
 
+			TArray<FString> ConfigFilesPluginsCannotOverride;
+			GConfig->GetArray(TEXT("Plugins"), TEXT("ConfigFilesPluginsCannotOverride"), ConfigFilesPluginsCannotOverride, GEngineIni);
+
 			TSet<FString> AllIniFiles;
 
 			if (PluginSystemDefs::IsCachingIniFilesForProcessing())
@@ -1050,7 +1053,7 @@ bool FPluginManager::ConfigureEnabledPlugins()
 			FCriticalSection PluginPakCS;
 			
 			// Mount all the enabled plugins
-			ParallelFor(PluginsArray.Num(), [&PluginsArray, &ConfigCS, &PluginPakCS, &AllIniFiles, this](int32 Index)
+			ParallelFor(PluginsArray.Num(), [&PluginsArray, &ConfigCS, &PluginPakCS, &ConfigFilesPluginsCannotOverride, &AllIniFiles, this](int32 Index)
 			{
 				FString PlatformName = FPlatformProperties::PlatformName();
 				FPlugin& Plugin = *PluginsArray[Index];
@@ -1124,6 +1127,16 @@ bool FPluginManager::ConfigureEnabledPlugins()
 						{
 							UE_LOG(LogPluginManager, Log, TEXT("Found config from plugin[%s] %s"), *Plugin.GetName(), *PluginConfigFilename);
 
+							if (ConfigFilesPluginsCannotOverride.Contains(BaseConfigFile))
+							{
+								// Not allowed, skip it
+								FText FailureMessage = FText::Format(LOCTEXT("PluginOverrideFailureFormat", "Plugin '{0}' cannot override config file: '{1}'"), FText::FromString(Plugin.GetName()), FText::FromString(BaseConfigFile));
+								FText DialogTitle = LOCTEXT("PluginConfigFileOverride", "Plugin config file override");
+								UE_LOG(LogPluginManager, Error, TEXT("%s"), *FailureMessage.ToString());
+								FMessageDialog::Open(EAppMsgType::Ok, FailureMessage, &DialogTitle);
+								continue;
+							}
+
 							FoundConfig->AddDynamicLayerToHierarchy(FPaths::Combine(PluginConfigDir, ConfigFile));
 
 #if ALLOW_INI_OVERRIDE_FROM_COMMANDLINE
