Add editor-only CachedJson and AdditionalFieldsToWrite to FPluginReferenceDescriptor similar to what we have in FPluginDescriptor

diff --git a/Engine/Source/Runtime/Projects/Private/PluginDescriptor.cpp b/Engine/Source/Runtime/Projects/Private/PluginDescriptor.cpp
--- a/Engine/Source/Runtime/Projects/Private/PluginDescriptor.cpp
+++ b/Engine/Source/Runtime/Projects/Private/PluginDescriptor.cpp
@@ -283,7 +283,7 @@ void FPluginDescriptor::Write(TJsonWriter<>& Writer) const
 	{
 		FJsonObject::Duplicate(/*Source=*/ CachedJson, /*Dest=*/ PluginJsonObject);
 	}
-#endif
+#endif //if WITH_EDITOR
 
 	UpdateJson(*PluginJsonObject);
 
@@ -432,7 +432,7 @@ void FPluginDescriptor::UpdateJson(FJsonObject& JsonObject) const
 	{
 		JsonObject.SetField(KVP.Key, FJsonValue::Duplicate(KVP.Value));
 	}
-#endif
+#endif //if WITH_EDITOR
 }
 
 bool FPluginDescriptor::UpdatePluginFile(const FString& FileName, FText* OutFailReason /*= nullptr*/) const
diff --git a/Engine/Source/Runtime/Projects/Private/PluginReferenceDescriptor.cpp b/Engine/Source/Runtime/Projects/Private/PluginReferenceDescriptor.cpp
--- a/Engine/Source/Runtime/Projects/Private/PluginReferenceDescriptor.cpp
+++ b/Engine/Source/Runtime/Projects/Private/PluginReferenceDescriptor.cpp
@@ -117,26 +117,40 @@ bool FPluginReferenceDescriptor::IsSupportedTargetPlatform(const FString& Platfo
 	}
 }
 
-bool FPluginReferenceDescriptor::Read(const FJsonObject& Object, FText* OutFailReason /*= nullptr*/)
+bool FPluginReferenceDescriptor::Read(const TSharedRef<FJsonObject>& Object, FText* OutFailReason /*= nullptr*/)
 {
+PRAGMA_DISABLE_DEPRECATION_WARNINGS
+	return Read(*Object, OutFailReason, Object);
+PRAGMA_ENABLE_DEPRECATION_WARNINGS
+}
+
+bool FPluginReferenceDescriptor::Read(const FJsonObject& Object, FText* OutFailReason /*= nullptr*/, TSharedPtr<FJsonObject> ObjectPtr /*= nullptr*/)
+{
+#if WITH_EDITOR
+	CachedJson = ObjectPtr;
+	AdditionalFieldsToWrite.Reset();
+#endif // WITH_EDITOR
+
+	bool bSucces = true;
+
 	// Get the name
-	if(!Object.TryGetStringField(TEXT("Name"), Name))
+	if (!Object.TryGetStringField(TEXT("Name"), Name))
 	{
 		if (OutFailReason)
 		{
 			*OutFailReason = LOCTEXT("PluginReferenceWithoutName", "Plugin references must have a 'Name' field");
 		}
-		return false;
+		bSucces = false;
 	}
 
 	// Get the enabled field
-	if(!Object.TryGetBoolField(TEXT("Enabled"), bEnabled))
+	if (!Object.TryGetBoolField(TEXT("Enabled"), bEnabled))
 	{
 		if (OutFailReason)
 		{
 			*OutFailReason = LOCTEXT("PluginReferenceWithoutEnabled", "Plugin references must have an 'Enabled' field");
 		}
-		return false;
+		bSucces = false;
 	}
 
 	// Read the optional field
@@ -162,34 +176,36 @@ bool FPluginReferenceDescriptor::Read(const FJsonObject& Object, FText* OutFailR
 	Object.TryGetStringArrayField(TEXT("SupportedTargetPlatforms"), SupportedTargetPlatforms);
 	Object.TryGetBoolField(TEXT("HasExplicitPlatforms"), bHasExplicitPlatforms);
 
-	return true;
+	return bSucces;
 }
 
-bool FPluginReferenceDescriptor::Read(const FJsonObject& Object, FText& OutFailReason)
+bool FPluginReferenceDescriptor::Read(const FJsonObject& Object, FText& OutFailReason, TSharedPtr<FJsonObject> ObjectPtr /*= nullptr*/)
 {
-	return Read(Object, &OutFailReason);
+PRAGMA_DISABLE_DEPRECATION_WARNINGS
+	return Read(Object, &OutFailReason, ObjectPtr);
+PRAGMA_ENABLE_DEPRECATION_WARNINGS
 }
 
 bool FPluginReferenceDescriptor::ReadArray(const FJsonObject& Object, const TCHAR* Name, TArray<FPluginReferenceDescriptor>& OutPlugins, FText* OutFailReason /*= nullptr*/)
 {
-	const TArray< TSharedPtr<FJsonValue> > *Array;
+	const TArray< TSharedPtr<FJsonValue> >* Array;
 
 	if (Object.TryGetArrayField(Name, Array))
 	{
 		for (const TSharedPtr<FJsonValue> &Item : *Array)
 		{
-			const TSharedPtr<FJsonObject> *ObjectPtr;
+			const TSharedPtr<FJsonObject>* ItemObject = nullptr;
 
-			if (Item.IsValid() && Item->TryGetObject(ObjectPtr))
+			if (Item.IsValid() && Item->TryGetObject(ItemObject) && ItemObject && ItemObject->IsValid())
 			{
-				FPluginReferenceDescriptor Plugin;
+				FPluginReferenceDescriptor PluginRef;
 
-				if (!Plugin.Read(*ObjectPtr->Get(), OutFailReason))
+				if (!PluginRef.Read(ItemObject->ToSharedRef(), OutFailReason))
 				{
 					return false;
 				}
 
-				OutPlugins.Add(Plugin);
+				OutPlugins.Add(PluginRef);
 			}
 		}
 	}
@@ -204,10 +220,18 @@ bool FPluginReferenceDescriptor::ReadArray(const FJsonObject& Object, const TCHA
 
 void FPluginReferenceDescriptor::Write(TJsonWriter<>& Writer) const
 {
-	TSharedRef<FJsonObject> PluginRefJsonObject = MakeShared<FJsonObject>();
+	TSharedPtr<FJsonObject> PluginRefJsonObject = MakeShared<FJsonObject>();
+
+#if WITH_EDITOR
+	if (CachedJson.IsValid())
+	{
+		FJsonObject::Duplicate(/*Source=*/ CachedJson, /*Dest=*/ PluginRefJsonObject);
+	}
+#endif //if WITH_EDITOR
+
 	UpdateJson(*PluginRefJsonObject);
 
-	FJsonSerializer::Serialize(PluginRefJsonObject, Writer);
+	FJsonSerializer::Serialize(PluginRefJsonObject.ToSharedRef(), Writer);
 }
 
 void FPluginReferenceDescriptor::UpdateJson(FJsonObject& JsonObject) const
@@ -356,6 +380,13 @@ void FPluginReferenceDescriptor::UpdateJson(FJsonObject& JsonObject) const
 	JsonObject.RemoveField(TEXT("BlacklistTargetConfigurations"));
 	JsonObject.RemoveField(TEXT("WhitelistTargets"));
 	JsonObject.RemoveField(TEXT("BlacklistTargets"));
+
+#if WITH_EDITOR
+	for (const auto& KVP : AdditionalFieldsToWrite)
+	{
+		JsonObject.SetField(KVP.Key, FJsonValue::Duplicate(KVP.Value));
+	}
+#endif //if WITH_EDITOR
 }
 
 void FPluginReferenceDescriptor::WriteArray(TJsonWriter<>& Writer, const TCHAR* ArrayName, const TArray<FPluginReferenceDescriptor>& Plugins)
@@ -384,4 +415,33 @@ void FPluginReferenceDescriptor::UpdateArray(FJsonObject& JsonObject, const TCHA
 		FPluginRefJsonArrayUpdater::FUpdateJsonObject::CreateStatic(PluginReferenceDescriptor::UpdatePluginRefJsonObject));
 }
 
+#if WITH_EDITOR
+bool FPluginReferenceDescriptor::GetAdditionalStringField(const FString& Key, FString& OutValue) const
+{
+	if (const TSharedPtr<FJsonValue>* ValuePtr = AdditionalFieldsToWrite.Find(Key))
+	{
+		const TSharedPtr<FJsonValue>& Value = *ValuePtr;
+		if (Value && Value->Type == EJson::String)
+		{
+			OutValue = Value->AsString();
+			return true;
+		}
+	}
+
+	if (CachedJson)
+	{
+		if (TSharedPtr<FJsonValue> Value = CachedJson->TryGetField(Key))
+		{
+			if (Value->Type == EJson::String)
+			{
+				OutValue = Value->AsString();
+				return true;
+			}
+		}
+	}
+
+	return false;
+}
+#endif //if WITH_EDITOR
+
 #undef LOCTEXT_NAMESPACE
diff --git a/Engine/Source/Runtime/Projects/Public/PluginReferenceDescriptor.h b/Engine/Source/Runtime/Projects/Public/PluginReferenceDescriptor.h
--- a/Engine/Source/Runtime/Projects/Public/PluginReferenceDescriptor.h
+++ b/Engine/Source/Runtime/Projects/Public/PluginReferenceDescriptor.h
@@ -7,6 +7,8 @@
 #include "CustomBuildSteps.h"
 #include "LocalizationDescriptor.h"
 
+class FJsonValue;
+
 /**
  * Descriptor for a plugin reference. Contains the information required to enable or disable a plugin for a given platform.
  */
@@ -51,6 +53,13 @@ struct PROJECTS_API FPluginReferenceDescriptor
 	/** When true, empty SupportedTargetPlatforms and PlatformAllowList are interpreted as 'no platforms' with the expectation that explicit platforms will be added in plugin platform extensions */
 	bool bHasExplicitPlatforms;
 
+#if WITH_EDITOR
+	/** Cached json for custom data */
+	mutable TSharedPtr<FJsonObject> CachedJson;
+
+	/** Additional fields to write */
+	TMap<FString, TSharedPtr<FJsonValue>> AdditionalFieldsToWrite;
+#endif //if WITH_EDITOR
 
 	/** Constructor */
 	FPluginReferenceDescriptor(const FString& InName = TEXT(""), bool bInEnabled = false);
@@ -68,10 +77,13 @@ struct PROJECTS_API FPluginReferenceDescriptor
 	bool IsSupportedTargetPlatform(const FString& Platform) const;
 
 	/** Reads the descriptor from the given JSON object */
-	bool Read(const FJsonObject& Object, FText* OutFailReason = nullptr);
+	bool Read(const TSharedRef<FJsonObject>& Object, FText* OutFailReason = nullptr);
 
-	/** Reads the descriptor from the given JSON object */
-	bool Read(const FJsonObject& Object, FText& OutFailReason);
+	UE_DEPRECATED(5.1, "Use Read(const TSharedRef<FJsonObject>&) instead.")
+	bool Read(const FJsonObject& Object, FText* OutFailReason = nullptr, TSharedPtr<FJsonObject> ObjectPtr = nullptr);
+
+	UE_DEPRECATED(5.1, "Use Read(const TSharedRef<FJsonObject>&) instead.")
+	bool Read(const FJsonObject& Object, FText& OutFailReason, TSharedPtr<FJsonObject> ObjectPtr = nullptr);
 
 	/** Reads an array of modules from the given JSON object */
 	static bool ReadArray(const FJsonObject& Object, const TCHAR* Name, TArray<FPluginReferenceDescriptor>& OutModules, FText* OutFailReason = nullptr);
@@ -90,4 +102,9 @@ struct PROJECTS_API FPluginReferenceDescriptor
 
 	/** Updates an array of plugin references in the specified JSON field (indexed by plugin name) */
 	static void UpdateArray(FJsonObject& JsonObject, const TCHAR* ArrayName, const TArray<FPluginReferenceDescriptor>& Plugins);
+
+#if WITH_EDITOR
+	/** Gets the string value for a given key by first looking into AdditionalFieldsToWrite and then CachedJson */
+	bool GetAdditionalStringField(const FString& Key, FString& OutValue) const;
+#endif //if WITH_EDITOR
 };
