Add IPluginManager::GetBuiltInPluginNames

diff --git a/Engine/Source/Runtime/Projects/Private/PluginManager.cpp b/Engine/Source/Runtime/Projects/Private/PluginManager.cpp
--- a/Engine/Source/Runtime/Projects/Private/PluginManager.cpp
+++ b/Engine/Source/Runtime/Projects/Private/PluginManager.cpp
@@ -329,6 +329,9 @@ void FPluginManager::DiscoverAllPlugins()
 	for (const TPair<FString, TSharedRef<FPlugin>>& PluginPair : AllPlugins)
 	{
 		PluginsToConfigure.Add(PluginPair.Key);
+#if WITH_EDITOR
+		BuiltInPluginNames.Add(PluginPair.Key);
+#endif //if WITH_EDITOR
 	}
 }
 
@@ -1891,6 +1894,14 @@ TArray<TSharedRef<IPlugin>> FPluginManager::GetDiscoveredPlugins()
 	return Plugins;
 }
 
+#if WITH_EDITOR
+const TSet<FString>& FPluginManager::GetBuiltInPluginNames() const
+{
+	ensure(!BuiltInPluginNames.IsEmpty());
+	return BuiltInPluginNames;
+}
+#endif //if WITH_EDITOR
+
 bool FPluginManager::AddPluginSearchPath(const FString& ExtraDiscoveryPath, bool bRefresh)
 {
 	bool bAlreadyExists = false;
diff --git a/Engine/Source/Runtime/Projects/Private/PluginManager.h b/Engine/Source/Runtime/Projects/Private/PluginManager.h
--- a/Engine/Source/Runtime/Projects/Private/PluginManager.h
+++ b/Engine/Source/Runtime/Projects/Private/PluginManager.h
@@ -137,6 +137,11 @@ public:
 	virtual TArray<TSharedRef<IPlugin>> GetEnabledPluginsWithVerse() const override;
 	virtual TArray<TSharedRef<IPlugin>> GetEnabledPluginsWithContentOrVerse() const override;
 	virtual TArray<TSharedRef<IPlugin>> GetDiscoveredPlugins() override;
+
+#if WITH_EDITOR
+	virtual const TSet<FString>& GetBuiltInPluginNames() const override;
+#endif //WITH_EDITOR
+
 	virtual bool AddPluginSearchPath(const FString& ExtraDiscoveryPath, bool bRefresh = true) override;
 	const TSet<FString>& GetAdditionalPluginSearchPaths() const override;
 	virtual TArray<TSharedRef<IPlugin>> GetPluginsWithPakFile() const override;
@@ -221,6 +226,11 @@ private:
 	/** Plugins that need to be configured to see if they should be enabled */
 	TSet<FString> PluginsToConfigure;
 
+#if WITH_EDITOR
+	/** Names of built-in plugins */
+	TSet<FString> BuiltInPluginNames;
+#endif //if WITH_EDITOR
+
 	TArray<TSharedRef<IPlugin>> PluginsWithPakFile;
 
 	/** Delegate for mounting content paths.  Bound by FPackageName code in CoreUObject, so that we can access
diff --git a/Engine/Source/Runtime/Projects/Public/Interfaces/IPluginManager.h b/Engine/Source/Runtime/Projects/Public/Interfaces/IPluginManager.h
--- a/Engine/Source/Runtime/Projects/Public/Interfaces/IPluginManager.h
+++ b/Engine/Source/Runtime/Projects/Public/Interfaces/IPluginManager.h
@@ -281,7 +281,7 @@ public:
 #endif
 
 	/**
-	 * Finds information for an enabled plugin.
+	 * Finds information for a plugin.
 	 *
 	 * @return	 Pointer to the plugin's information, or nullptr.
 	 */
@@ -323,6 +323,13 @@ public:
 	 */
 	virtual TArray<TSharedRef<IPlugin>> GetDiscoveredPlugins() = 0;
 
+#if WITH_EDITOR
+	/**
+	 * Returns the set of built-in plugin names
+	 */
+	PROJECTS_API virtual const TSet<FString>& GetBuiltInPluginNames() const = 0;
+#endif //WITH_EDITOR
+
 	/**
 	 * Stores the specified path, utilizing it in future search passes when 
 	 * searching for available plugins. Optionally refreshes the manager after 
