Game Feature Plugins get removed from plugin manager in FGameFeaturePluginState_Unmounting if they had to be added to the plugin manager list by FGameFeaturePluginState_Mounting.

diff --git a/Engine/Plugins/Experimental/GameFeatures/Source/GameFeatures/Private/GameFeaturePluginStateMachine.cpp b/Engine/Plugins/Experimental/GameFeatures/Source/GameFeatures/Private/GameFeaturePluginStateMachine.cpp
--- a/Engine/Plugins/Experimental/GameFeatures/Source/GameFeatures/Private/GameFeaturePluginStateMachine.cpp
+++ b/Engine/Plugins/Experimental/GameFeatures/Source/GameFeatures/Private/GameFeaturePluginStateMachine.cpp
@@ -786,11 +786,18 @@ struct FGameFeaturePluginState_Unmounting : public FGameFeaturePluginState
 		PendingBundles.Empty();
 		bUnmounted = false;
 
-		if (TSharedPtr<IPlugin> Plugin = IPluginManager::Get().FindPlugin(StateProperties.PluginName))
+		if (TSharedPtr<IPlugin> Plugin = IPluginManager::Get().FindPlugin(StateProperties.PluginName); 
+			Plugin && Plugin->GetDescriptor().bExplicitlyLoaded)
 		{
 			// The asset registry listens to FPackageName::OnContentPathDismounted() and 
 			// will automatically cleanup the asset registry state we added for this plugin.
-			IPluginManager::Get().UnmountExplicitlyLoadedPlugin(StateProperties.PluginName, nullptr);
+			verify(IPluginManager::Get().UnmountExplicitlyLoadedPlugin(StateProperties.PluginName, nullptr));
+		}
+
+		if (StateProperties.bAddedPluginToManager)
+		{
+			verify(IPluginManager::Get().RemoveFromPluginsList(StateProperties.PluginName));
+			StateProperties.bAddedPluginToManager = false;
 		}
 
 		if (StateProperties.GetPluginProtocol() != EGameFeaturePluginProtocol::InstallBundle)
@@ -985,14 +992,25 @@ struct FGameFeaturePluginState_Mounting : public FGameFeaturePluginState
 		checkf(FPaths::GetExtension(StateProperties.PluginInstalledFilename) == TEXT("uplugin"), TEXT("PluginInstalledFilename must have a uplugin extension. PluginURL: %s"), *StateProperties.PluginURL);
 
 		// refresh the plugins list to let the plugin manager know about it
-		const bool bAddedPlugin = IPluginManager::Get().AddToPluginsList(StateProperties.PluginInstalledFilename);
-		if (!bAddedPlugin)
+		TSharedPtr<IPlugin> MaybePlugin = IPluginManager::Get().FindPlugin(StateProperties.PluginName);
+		const bool bNeedsPluginMount = (MaybePlugin == nullptr || MaybePlugin->GetDescriptor().bExplicitlyLoaded);
+
+		if (MaybePlugin == nullptr)
 		{
-			StateStatus.SetTransitionError(EGameFeaturePluginState::ErrorMounting, UE::GameFeatures::StateMachineErrorNamespace + TEXT("Failed_To_Register_Plugin"));
-			return;
+			const bool bAddedPlugin = IPluginManager::Get().AddToPluginsList(StateProperties.PluginInstalledFilename);
+			if (!bAddedPlugin)
+			{
+				StateStatus.SetTransitionError(EGameFeaturePluginState::ErrorMounting, UE::GameFeatures::StateMachineErrorNamespace + TEXT("Failed_To_Register_Plugin"));
+				return;
+			}
+
+			StateProperties.bAddedPluginToManager = true;
 		}
 		
-		IPluginManager::Get().MountExplicitlyLoadedPlugin(StateProperties.PluginName);
+		if (bNeedsPluginMount)
+		{
+			IPluginManager::Get().MountExplicitlyLoadedPlugin(StateProperties.PluginName);
+		}
 
 		// After the new plugin is mounted add the asset registry for that plugin.
 		if (StateProperties.GetPluginProtocol() == EGameFeaturePluginProtocol::InstallBundle)
diff --git a/Engine/Plugins/Experimental/GameFeatures/Source/GameFeatures/Private/GameFeaturePluginStateMachine.h b/Engine/Plugins/Experimental/GameFeatures/Source/GameFeatures/Private/GameFeaturePluginStateMachine.h
--- a/Engine/Plugins/Experimental/GameFeatures/Source/GameFeatures/Private/GameFeaturePluginStateMachine.h
+++ b/Engine/Plugins/Experimental/GameFeatures/Source/GameFeatures/Private/GameFeaturePluginStateMachine.h
@@ -177,15 +177,18 @@ struct FGameFeaturePluginStateMachineProperties
 	 */
 	FString PluginURL;
 
-	/** Once installed, this is the filename on disk of the .uplugin file. */
+	/** Filename on disk of the .uplugin file. */
 	FString PluginInstalledFilename;
 	
-	/** Once installed, this is the name of the plugin. */
+	/** Name of the plugin. */
 	FString PluginName;
 
 	/** Meta data parsed from the URL for a specific protocol. */
 	TUnion<FInstallBundlePluginProtocolMetaData> ProtocolMetadata;
 
+	/** Tracks whether or not this state machine added the pluging to the plugin manager. */
+	bool bAddedPluginToManager = false;
+
 	/** The desired state during a transition. */
 	EGameFeaturePluginState DestinationState = EGameFeaturePluginState::Uninitialized;
 
diff --git a/Engine/Source/Runtime/Projects/Private/PluginManager.cpp b/Engine/Source/Runtime/Projects/Private/PluginManager.cpp
--- a/Engine/Source/Runtime/Projects/Private/PluginManager.cpp
+++ b/Engine/Source/Runtime/Projects/Private/PluginManager.cpp
@@ -254,6 +254,23 @@ void FPluginManager::RefreshPluginsList()
 	}
 }
 
+bool VerifyForSinglePluginForAddOrRemove(const FPluginDescriptor& Descriptor, FText& OutFailReason)
+{
+	if (!ensureMsgf(Descriptor.bExplicitlyLoaded, TEXT("VerifyForSinglePluginForAddOrRemove does not allow plugins with bIsPluginExtension set to true. Plugin: %s"), *Descriptor.FriendlyName))
+	{
+		OutFailReason = FText::Format(LOCTEXT("PluginIsNotExplicitlyLoaded", "Plugin '{0}' is not explicitly loaded"), FText::FromString(Descriptor.FriendlyName));
+		return false;
+	}
+
+	if (!ensureMsgf(!Descriptor.bIsPluginExtension, TEXT("VerifyForSinglePluginForAddOrRemove does not allow platform extensions. Plugin %s"), *Descriptor.FriendlyName))
+	{
+		OutFailReason = FText::Format(LOCTEXT("PluginIsExtension", "Plugin '{0}' is a platform extension"), FText::FromString(Descriptor.FriendlyName));
+		return false;
+	}
+
+	return true;
+}
+
 bool FPluginManager::AddToPluginsList(const FString& PluginFilename, FText* OutFailReason /*= nullptr*/)
 {
 #if (WITH_ENGINE && !IS_PROGRAM) || WITH_PLUGIN_SUPPORT
@@ -267,7 +284,13 @@ bool FPluginManager::AddToPluginsList(const FString& PluginFilename, FText* OutF
 	// Read the plugin and load it
 	FPluginDescriptor Descriptor;
 	FText FailReason;
-	if (Descriptor.Load(PluginFilename, FailReason))
+	bool bSuccess = Descriptor.Load(PluginFilename, FailReason);
+	if (bSuccess)
+	{
+		bSuccess = VerifyForSinglePluginForAddOrRemove(Descriptor, FailReason);
+	}
+
+	if(bSuccess)
 	{
 		// Determine the plugin type
 		EPluginType PluginType = EPluginType::External;
@@ -318,6 +341,47 @@ bool FPluginManager::AddToPluginsList(const FString& PluginFilename, FText* OutF
 	return false;
 }
 
+bool FPluginManager::RemoveFromPluginsList(const FString& PluginFilename, FText* OutFailReason /*= nullptr*/)
+{
+#if (WITH_ENGINE && !IS_PROGRAM) || WITH_PLUGIN_SUPPORT
+	FString PluginName = FPaths::GetBaseFilename(PluginFilename);
+
+	TSharedRef<FPlugin>* MaybePlugin = AllPlugins.Find(PluginName);
+	if (MaybePlugin == nullptr)
+	{
+		return true;
+	}
+
+	const TSharedRef<FPlugin>& FoundPlugin = *MaybePlugin;
+
+	FText FailReason;
+	bool bSuccess = true;
+	if (FoundPlugin->IsEnabled())
+	{
+		FailReason = FText::Format(LOCTEXT("PluginIsEnabled", "Plugin '{0}' is enabled"), FText::FromString(PluginFilename));
+		bSuccess = false;
+	}
+
+	if(bSuccess)
+	{
+		bSuccess = VerifyForSinglePluginForAddOrRemove(FoundPlugin->Descriptor, FailReason);
+	}
+
+	if (!bSuccess)
+	{
+		if (OutFailReason)
+		{
+			*OutFailReason = FText::Format(LOCTEXT("RemoveFromPluginsListFailed", "Failed to remove plugin '{0}'\n{1}"), FText::FromString(PluginFilename), FailReason);
+		}
+		return false;
+	}
+
+	AllPlugins.Remove(PluginName);
+#endif
+
+	return true;
+}
+
 void FPluginManager::DiscoverAllPlugins()
 {
 	ensure( AllPlugins.Num() == 0 );		// Should not have already been initialized!
@@ -974,7 +1038,7 @@ bool FPluginManager::ConfigureEnabledPlugins()
 		// If we made it here, we have all the required plugins
 		bHaveAllRequiredPlugins = true;
 
-		// check if the config already contais the plugin inis - if so, we don't need to scan anything, just use the ini to find paks to mount
+		// check if the config already contains the plugin inis - if so, we don't need to scan anything, just use the ini to find paks to mount
 		TArray<FString> BinaryConfigPlugins;
 		if (GConfig->GetArray(TEXT("BinaryConfig"), TEXT("BinaryConfigPlugins"), BinaryConfigPlugins, GEngineIni) && BinaryConfigPlugins.Num() > 0)
 		{
diff --git a/Engine/Source/Runtime/Projects/Private/PluginManager.h b/Engine/Source/Runtime/Projects/Private/PluginManager.h
--- a/Engine/Source/Runtime/Projects/Private/PluginManager.h
+++ b/Engine/Source/Runtime/Projects/Private/PluginManager.h
@@ -111,6 +111,7 @@ public:
 	/** IPluginManager interface */
 	virtual void RefreshPluginsList() override;
 	virtual bool AddToPluginsList(const FString& PluginFilename, FText* OutFailReason = nullptr) override;
+	virtual bool RemoveFromPluginsList(const FString& PluginFilename, FText* OutFailReason = nullptr) override;
 	virtual bool LoadModulesForEnabledPlugins( const ELoadingPhase::Type LoadingPhase ) override;
 	virtual FLoadingModulesForPhaseEvent& OnLoadingPhaseComplete() override;
 	virtual ELoadingPhase::Type GetLastCompletedLoadingPhase() const override;
diff --git a/Engine/Source/Runtime/Projects/Public/Interfaces/IPluginManager.h b/Engine/Source/Runtime/Projects/Public/Interfaces/IPluginManager.h
--- a/Engine/Source/Runtime/Projects/Public/Interfaces/IPluginManager.h
+++ b/Engine/Source/Runtime/Projects/Public/Interfaces/IPluginManager.h
@@ -208,6 +208,13 @@ public:
 	 */
 	virtual bool AddToPluginsList(const FString& PluginFilename, FText* OutFailReason = nullptr) = 0;
 
+	/**
+	 * Remove a single plugin from the list of plugins.
+	 *
+	 * @return True if the plugin was not in the list. False if it can't be removed (see OutFailReason).
+	 */
+	virtual bool RemoveFromPluginsList(const FString& PluginFilename, FText* OutFailReason = nullptr) = 0;
+
 	/**
 	 * Loads all plug-ins
 	 *
