Fix typos and remove check for bExplicitlyLoaded.  Its ok to add a plugin to the list that doesn't have bExplicitlyLoaded.

diff --git a/Engine/Source/Runtime/Projects/Private/PluginManager.cpp b/Engine/Source/Runtime/Projects/Private/PluginManager.cpp
--- a/Engine/Source/Runtime/Projects/Private/PluginManager.cpp
+++ b/Engine/Source/Runtime/Projects/Private/PluginManager.cpp
@@ -254,15 +254,9 @@ void FPluginManager::RefreshPluginsList()
 	}
 }
 
-bool VerifyForSinglePluginForAddOrRemove(const FPluginDescriptor& Descriptor, FText& OutFailReason)
+bool VerifySinglePluginForAddOrRemove(const FPluginDescriptor& Descriptor, FText& OutFailReason)
 {
-	if (!ensureMsgf(Descriptor.bExplicitlyLoaded, TEXT("VerifyForSinglePluginForAddOrRemove does not allow plugins with bIsPluginExtension set to true. Plugin: %s"), *Descriptor.FriendlyName))
-	{
-		OutFailReason = FText::Format(LOCTEXT("PluginIsNotExplicitlyLoaded", "Plugin '{0}' is not explicitly loaded"), FText::FromString(Descriptor.FriendlyName));
-		return false;
-	}
-
-	if (!ensureMsgf(!Descriptor.bIsPluginExtension, TEXT("VerifyForSinglePluginForAddOrRemove does not allow platform extensions. Plugin %s"), *Descriptor.FriendlyName))
+	if (!ensureMsgf(!Descriptor.bIsPluginExtension, TEXT("VerifySinglePluginForAddOrRemove does not allow platform extensions. Plugin %s"), *Descriptor.FriendlyName))
 	{
 		OutFailReason = FText::Format(LOCTEXT("PluginIsExtension", "Plugin '{0}' is a platform extension"), FText::FromString(Descriptor.FriendlyName));
 		return false;
@@ -287,7 +281,7 @@ bool FPluginManager::AddToPluginsList(const FString& PluginFilename, FText* OutF
 	bool bSuccess = Descriptor.Load(PluginFilename, FailReason);
 	if (bSuccess)
 	{
-		bSuccess = VerifyForSinglePluginForAddOrRemove(Descriptor, FailReason);
+		bSuccess = VerifySinglePluginForAddOrRemove(Descriptor, FailReason);
 	}
 
 	if(bSuccess)
@@ -364,7 +358,7 @@ bool FPluginManager::RemoveFromPluginsList(const FString& PluginFilename, FText*
 
 	if(bSuccess)
 	{
-		bSuccess = VerifyForSinglePluginForAddOrRemove(FoundPlugin->Descriptor, FailReason);
+		bSuccess = VerifySinglePluginForAddOrRemove(FoundPlugin->Descriptor, FailReason);
 	}
 
 	if (!bSuccess)
