Improve the warning logged when we detect a plugin using an engine version that is not compatible with the current engine version.

#rnx
PR #8567: Improve diagnostics for plugin version mismatch. (Contributed by [REDACTED_PERSON])

- I tweeked the log message slightly from the PR
-Worth noting that we don't need to localize this message even if it does feel a bit user facing as ::PromptToLoadIncompatiblePlugin will display a localized error message relating to the problem to the end user.

diff --git a/Engine/Source/Runtime/Projects/Private/PluginManager.cpp b/Engine/Source/Runtime/Projects/Private/PluginManager.cpp
--- a/Engine/Source/Runtime/Projects/Private/PluginManager.cpp
+++ b/Engine/Source/Runtime/Projects/Private/PluginManager.cpp
@@ -1606,7 +1606,7 @@ bool FPluginManager::IsPluginCompatible(const FPlugin& Plugin)
 		EVersionComparison Comparison = FEngineVersion::GetNewest(FEngineVersion::CompatibleWith(), Version, nullptr);
 		if (Comparison != EVersionComparison::Neither)
 		{
-			UE_LOG(LogPluginManager, Warning, TEXT("Plugin '%s' is not compatible with the current engine version (%s)"), *Plugin.Name, *Plugin.Descriptor.EngineVersion);
+			UE_LOG(LogPluginManager, Warning, TEXT("Plugin '%s' requires engine version '%s' and may not be compatible with the current current engine version '%s'"), *Plugin.Name, *Plugin.Descriptor.EngineVersion, *FEngineVersion::CompatibleWith().ToString());
 			return false;
 		}
 	}
