FProjectDescriptor::HasModule to quickly test whether a given module is a project module (editor only)

diff --git a/Engine/Source/Runtime/Projects/Private/ProjectDescriptor.cpp b/Engine/Source/Runtime/Projects/Private/ProjectDescriptor.cpp
--- a/Engine/Source/Runtime/Projects/Private/ProjectDescriptor.cpp
+++ b/Engine/Source/Runtime/Projects/Private/ProjectDescriptor.cpp
@@ -110,6 +110,10 @@ bool FProjectDescriptor::Read(const FJsonObject& Object, const FString& PathToPr
 	Object.TryGetBoolField(TEXT("Enterprise"), bIsEnterpriseProject);
 	Object.TryGetBoolField(TEXT("DisableEnginePluginsByDefault"), bDisableEnginePluginsByDefault);
 
+#if WITH_EDITOR
+	ModuleNamesCache.Reset();
+#endif
+
 	// Read the modules
 	if(!FModuleDescriptor::ReadArray(Object, TEXT("Modules"), Modules, OutFailReason))
 	{
@@ -390,4 +394,21 @@ bool FProjectDescriptor::RemoveRootDirectory(const FString& Dir)
 	return AdditionalRootDirectories.RemoveSingle(Dir) > 0;
 }
 
+#if WITH_EDITOR
+bool FProjectDescriptor::HasModule(FName ModuleName) const
+{
+	if (ModuleNamesCache.Num() != Modules.Num())
+	{
+		ModuleNamesCache.Reset();
+		ModuleNamesCache.Reserve(Modules.Num());
+		for (const FModuleDescriptor& Module : Modules)
+		{
+			ModuleNamesCache.Add(Module.Name);
+		}
+		ensure(ModuleNamesCache.Num() == Modules.Num());
+	}
+	return ModuleNamesCache.Contains(ModuleName);
+}
+#endif //if WITH_EDITOR
+
 #undef LOCTEXT_NAMESPACE
diff --git a/Engine/Source/Runtime/Projects/Public/ProjectDescriptor.h b/Engine/Source/Runtime/Projects/Public/ProjectDescriptor.h
--- a/Engine/Source/Runtime/Projects/Public/ProjectDescriptor.h
+++ b/Engine/Source/Runtime/Projects/Public/ProjectDescriptor.h
@@ -169,6 +169,10 @@ struct PROJECTS_API FProjectDescriptor
 	 */
 	bool RemoveRootDirectory(const FString& Dir);
 
+#if WITH_EDITOR
+	/** Returns whether the project has a module of the given name */
+	bool HasModule(FName ModuleName) const;
+#endif //if WITH_EDITOR
 
 private:
 	/** @return the path relative to this project if possible */
@@ -185,4 +189,9 @@ private:
 	 * Paths are in memory as absolute paths. Conversion to/from path relative happens during Save/Load
 	 */
 	TArray<FString> AdditionalRootDirectories;
+
+#if WITH_EDITOR
+	/** Set of project module names */
+	mutable TSet<FName> ModuleNamesCache;
+#endif //if WITH_EDITOR
 };
