When unmounting a plugin, remove it from [Core.System]:Paths.

diff --git a/Engine/Source/Runtime/Projects/Private/PluginManager.cpp b/Engine/Source/Runtime/Projects/Private/PluginManager.cpp
--- a/Engine/Source/Runtime/Projects/Private/PluginManager.cpp
+++ b/Engine/Source/Runtime/Projects/Private/PluginManager.cpp
@@ -2097,6 +2097,18 @@ bool FPluginManager::UnmountPluginFromExternalSource(const TSharedPtr<FPlugin>&
 
 	if ((Plugin->CanContainContent() || Plugin->CanContainVerse()) && ensure(UnRegisterMountPointDelegate.IsBound()))
 	{
+		// Remove this plugin's path from the list of content directories that the editor will search
+		if (Plugin->CanContainContent())
+		{
+			if (FConfigFile* EngineConfigFile = GConfig->Find(GEngineIni))
+			{
+				if (FConfigSection* CoreSystemSection = EngineConfigFile->Find(TEXT("Core.System")))
+				{
+					CoreSystemSection->Remove("Paths", Plugin->GetContentDir());
+				}
+			}
+		}
+
 		UnRegisterMountPointDelegate.Execute(Plugin->GetMountedAssetPath(), Plugin->GetContentDir());
 	}
 
