Module name to plugin lookup table in plugin manager

diff --git a/Engine/Source/Runtime/Projects/Private/PluginManager.cpp b/Engine/Source/Runtime/Projects/Private/PluginManager.cpp
--- a/Engine/Source/Runtime/Projects/Private/PluginManager.cpp
+++ b/Engine/Source/Runtime/Projects/Private/PluginManager.cpp
@@ -253,6 +253,14 @@ void FPluginManager::RefreshPluginsList()
 			PluginsToConfigure.AddByHash(PluginNameHash, NewPlugin->GetName());
 		}
 	}
+
+#if WITH_EDITOR
+	ModuleNameToPluginMap.Reset();
+	for (const TPair<FString, TSharedRef<FPlugin>>& PluginPair : AllPlugins)
+	{
+		AddToModuleNameToPluginMap(PluginPair.Value);
+	}
+#endif //if WITH_EDITOR
 }
 
 bool VerifySinglePluginForAddOrRemove(const FPluginDescriptor& Descriptor, FText& OutFailReason)
@@ -318,6 +326,10 @@ bool FPluginManager::AddToPluginsList(const FString& PluginFilename, FText* OutF
 		if (ensure(NewPlugin))
 		{
 			AllPlugins.Add(PluginName, *NewPlugin);
+
+#if WITH_EDITOR
+			AddToModuleNameToPluginMap(*NewPlugin);
+#endif //if WITH_EDITOR
 		}
 
 		return true;
@@ -371,8 +383,12 @@ bool FPluginManager::RemoveFromPluginsList(const FString& PluginFilename, FText*
 		return false;
 	}
 
+#if WITH_EDITOR
+	RemoveFromModuleNameToPluginMap(FoundPlugin);
+#endif //if WITH_EDITOR
+
 	AllPlugins.Remove(PluginName);
-#endif
+#endif //if (WITH_ENGINE && !IS_PROGRAM) || WITH_PLUGIN_SUPPORT
 
 	return true;
 }
@@ -390,6 +406,7 @@ void FPluginManager::DiscoverAllPlugins()
 		PluginsToConfigure.Add(PluginPair.Key);
 #if WITH_EDITOR
 		BuiltInPluginNames.Add(PluginPair.Key);
+		AddToModuleNameToPluginMap(PluginPair.Value);
 #endif //if WITH_EDITOR
 	}
 }
@@ -1941,6 +1958,15 @@ const TSet<FString>& FPluginManager::GetBuiltInPluginNames() const
 	ensure(!BuiltInPluginNames.IsEmpty());
 	return BuiltInPluginNames;
 }
+
+TSharedPtr<IPlugin> FPluginManager::GetModuleOwnerPlugin(FName ModuleName) const
+{
+	if (const TSharedRef<IPlugin>* Plugin = ModuleNameToPluginMap.Find(ModuleName))
+	{
+		return *Plugin;
+	}
+	return TSharedPtr<IPlugin>();
+}
 #endif //if WITH_EDITOR
 
 bool FPluginManager::AddPluginSearchPath(const FString& ExtraDiscoveryPath, bool bRefresh)
@@ -2137,5 +2163,42 @@ FName FPluginManager::PackageNameFromModuleName(FName ModuleName)
 	return Result;
 }
 
+#if WITH_EDITOR
+void FPluginManager::AddToModuleNameToPluginMap(const TSharedRef<FPlugin>& Plugin)
+{
+	if (!Plugin->Descriptor.bIsPluginExtension)
+	{
+		for (const FModuleDescriptor& Module : Plugin->Descriptor.Modules)
+		{
+			if (const TSharedRef<IPlugin>* FoundPlugin = ModuleNameToPluginMap.Find(Module.Name))
+			{
+				if (*FoundPlugin != Plugin)
+				{
+					UE_LOG(LogPluginManager, Error, TEXT("Module %s from plugin %s is already associated with plugin %s (maybe because the plugin should be using bIsPluginExtension)"), *Module.Name.ToString(), *Plugin->GetName(), *(*FoundPlugin)->GetName());
+				}
+			}
+			else
+			{
+				ModuleNameToPluginMap.Add(Module.Name, Plugin);
+			}
+		}
+	}
+}
+
+void FPluginManager::RemoveFromModuleNameToPluginMap(const TSharedRef<FPlugin>& Plugin)
+{
+	if (!Plugin->Descriptor.bIsPluginExtension)
+	{
+		for (const FModuleDescriptor& Module : Plugin->Descriptor.Modules)
+		{
+			TSharedRef<IPlugin> RemovedPlugin(Plugin);
+			if (ModuleNameToPluginMap.RemoveAndCopyValue(Module.Name, RemovedPlugin))
+			{
+				ensure(RemovedPlugin == Plugin);
+			}
+		}
+	}
+}
+#endif //if WITH_EDITOR
 
 #undef LOCTEXT_NAMESPACE
diff --git a/Engine/Source/Runtime/Projects/Private/PluginManager.h b/Engine/Source/Runtime/Projects/Private/PluginManager.h
--- a/Engine/Source/Runtime/Projects/Private/PluginManager.h
+++ b/Engine/Source/Runtime/Projects/Private/PluginManager.h
@@ -141,6 +141,7 @@ public:
 
 #if WITH_EDITOR
 	virtual const TSet<FString>& GetBuiltInPluginNames() const override;
+	virtual TSharedPtr<IPlugin> GetModuleOwnerPlugin(FName ModuleName) const override;
 #endif //WITH_EDITOR
 
 	virtual bool AddPluginSearchPath(const FString& ExtraDiscoveryPath, bool bRefresh = true) override;
@@ -220,6 +221,11 @@ private:
 	/** Unmounts a plugin that was requested to be unmounted from external code (by UnmountExplicitlyLoadedPlugin) */
 	bool UnmountPluginFromExternalSource(const TSharedPtr<FPlugin>& Plugin, FText* OutReason);
 
+#if WITH_EDITOR
+	void AddToModuleNameToPluginMap(const TSharedRef<FPlugin>& Plugin);
+	void RemoveFromModuleNameToPluginMap(const TSharedRef<FPlugin>& Plugin);
+#endif //if WITH_EDITOR
+
 private:
 	/** All of the plugins that we know about */
 	TMap< FString, TSharedRef< FPlugin > > AllPlugins;
@@ -230,6 +236,8 @@ private:
 #if WITH_EDITOR
 	/** Names of built-in plugins */
 	TSet<FString> BuiltInPluginNames;
+
+	TMap<FName, TSharedRef<IPlugin>> ModuleNameToPluginMap;
 #endif //if WITH_EDITOR
 
 	TArray<TSharedRef<IPlugin>> PluginsWithPakFile;
diff --git a/Engine/Source/Runtime/Projects/Public/Interfaces/IPluginManager.h b/Engine/Source/Runtime/Projects/Public/Interfaces/IPluginManager.h
--- a/Engine/Source/Runtime/Projects/Public/Interfaces/IPluginManager.h
+++ b/Engine/Source/Runtime/Projects/Public/Interfaces/IPluginManager.h
@@ -335,6 +335,11 @@ public:
 	 * Returns the set of built-in plugin names
 	 */
 	PROJECTS_API virtual const TSet<FString>& GetBuiltInPluginNames() const = 0;
+
+	/**
+	 * Returns the plugin that owns the specified module, if any
+	 */
+	PROJECTS_API virtual TSharedPtr<IPlugin> GetModuleOwnerPlugin(FName ModuleName) const = 0;
 #endif //WITH_EDITOR
 
 	/**
