Remove log error from FPluginManager::AddToModuleNameToPluginMap until UE5 occurences are fixed

diff --git a/Engine/Source/Runtime/Projects/Private/PluginManager.cpp b/Engine/Source/Runtime/Projects/Private/PluginManager.cpp
--- a/Engine/Source/Runtime/Projects/Private/PluginManager.cpp
+++ b/Engine/Source/Runtime/Projects/Private/PluginManager.cpp
@@ -2174,7 +2174,7 @@ void FPluginManager::AddToModuleNameToPluginMap(const TSharedRef<FPlugin>& Plugi
 			{
 				if (*FoundPlugin != Plugin)
 				{
-					UE_LOG(LogPluginManager, Error, TEXT("Module %s from plugin %s is already associated with plugin %s (maybe because the plugin should be using bIsPluginExtension)"), *Module.Name.ToString(), *Plugin->GetName(), *(*FoundPlugin)->GetName());
+					UE_LOG(LogPluginManager, Display, TEXT("Module %s from plugin %s is already associated with plugin %s (maybe because the plugin should be using bIsPluginExtension)"), *Module.Name.ToString(), *Plugin->GetName(), *(*FoundPlugin)->GetName());
 				}
 			}
 			else
