Added new properties to Build.cs and uplugin files to support Verse public API
- new property bVersePublicApi in Build.cs files that is set to expose the associated Verse code to public users
- new property "VersePublicApi" in uplugin files that is set to expose the associated Verse code to public users

#rnx

diff --git a/Engine/Source/Programs/UnrealBuildTool/Configuration/ModuleRules.cs b/Engine/Source/Programs/UnrealBuildTool/Configuration/ModuleRules.cs
--- a/Engine/Source/Programs/UnrealBuildTool/Configuration/ModuleRules.cs
+++ b/Engine/Source/Programs/UnrealBuildTool/Configuration/ModuleRules.cs
@@ -1046,6 +1046,11 @@
 		public string? VersePath;
 
 		/// <summary>
+		/// If Verse code associated with this module is allowed to be seen by users outside of Epic
+		/// </summary>
+		public bool bVersePublicApi;
+
+		/// <summary>
 		/// Whether this module qualifies included headers from other modules relative to the root of their 'Public' folder. This reduces the number
 		/// of search paths that have to be passed to the compiler, improving performance and reducing the length of the compiler command line.
 		/// </summary>
diff --git a/Engine/Source/Programs/UnrealBuildTool/System/PluginDescriptor.cs b/Engine/Source/Programs/UnrealBuildTool/System/PluginDescriptor.cs
--- a/Engine/Source/Programs/UnrealBuildTool/System/PluginDescriptor.cs
+++ b/Engine/Source/Programs/UnrealBuildTool/System/PluginDescriptor.cs
@@ -144,6 +144,11 @@ namespace UnrealBuildTool
 		/// </summary>
 		public string? VersePath;
 
+		/// <summary>
+		/// If Verse code associated with this plugin is allowed to be seen by users outside of Epic
+		/// </summary>
+		public bool bVersePublicApi;
+
 		/// <summary>
 		/// Whether this plugin should be enabled by default for all projects
 		/// </summary>
@@ -294,6 +299,7 @@ namespace UnrealBuildTool
 			}
 
 			RawObject.TryGetStringField("VersePath", out VersePath);
+			RawObject.TryGetBoolField("VersePublicApi", out bVersePublicApi);
 
 			bool bEnabledByDefaultValue;
 			if(RawObject.TryGetBoolField("EnabledByDefault", out bEnabledByDefaultValue))
@@ -393,6 +399,10 @@ namespace UnrealBuildTool
 			{
 				Writer.WriteValue("VersePath", VersePath);
 			}
+			if (bVersePublicApi)
+			{
+				Writer.WriteValue("VersePublicApi", bVersePublicApi);
+			}
 			if (bEnabledByDefault.HasValue)
 			{
 				Writer.WriteValue("EnabledByDefault", bEnabledByDefault.Value);
diff --git a/Engine/Source/Programs/UnrealBuildTool/System/TargetMakefile.cs b/Engine/Source/Programs/UnrealBuildTool/System/TargetMakefile.cs
--- a/Engine/Source/Programs/UnrealBuildTool/System/TargetMakefile.cs
+++ b/Engine/Source/Programs/UnrealBuildTool/System/TargetMakefile.cs
@@ -101,7 +101,7 @@
 		/// <summary>
 		/// The version number to write
 		/// </summary>
-		public const int CurrentVersion = 30;
+		public const int CurrentVersion = 31;
 
 		/// <summary>
 		/// The time at which the makefile was created
diff --git a/Engine/Source/Runtime/Projects/Private/PluginDescriptor.cpp b/Engine/Source/Runtime/Projects/Private/PluginDescriptor.cpp
--- a/Engine/Source/Runtime/Projects/Private/PluginDescriptor.cpp
+++ b/Engine/Source/Runtime/Projects/Private/PluginDescriptor.cpp
@@ -77,6 +77,7 @@ const FString& FPluginDescriptor::GetFileExtension()
 
 FPluginDescriptor::FPluginDescriptor()
 	: Version(0)
+	, bVersePublicApi(false)
 	, EnabledByDefault(EPluginEnabledByDefault::Unspecified)
 	, bCanContainContent(false)
 	, bCanContainVerse(false)
@@ -211,11 +212,12 @@ bool FPluginDescriptor::Read(const FJsonObject& Object, FText* OutFailReason /*=
 	}
 
 	Object.TryGetStringField(TEXT("VersePath"), VersePath);
+	Object.TryGetBoolField(TEXT("VersePublicApi"), bVersePublicApi);
 
 	bool bEnabledByDefault;
 	if(Object.TryGetBoolField(TEXT("EnabledByDefault"), bEnabledByDefault))
 	{
-		EnabledByDefault = bEnabledByDefault? EPluginEnabledByDefault::Enabled : EPluginEnabledByDefault::Disabled;
+		EnabledByDefault = bEnabledByDefault ? EPluginEnabledByDefault::Enabled : EPluginEnabledByDefault::Disabled;
 	}
 
 	Object.TryGetBoolField(TEXT("CanContainContent"), bCanContainContent);
@@ -327,6 +329,11 @@ void FPluginDescriptor::UpdateJson(FJsonObject& JsonObject) const
 		JsonObject.SetStringField(TEXT("VersePath"), VersePath);
 	}
 
+	if (bVersePublicApi)
+	{
+		JsonObject.SetBoolField(TEXT("VersePublicApi"), bVersePublicApi);
+	}
+
 	if (EnabledByDefault != EPluginEnabledByDefault::Unspecified)
 	{
 		JsonObject.SetBoolField(TEXT("EnabledByDefault"), (EnabledByDefault == EPluginEnabledByDefault::Enabled));
diff --git a/Engine/Source/Runtime/Projects/Public/PluginDescriptor.h b/Engine/Source/Runtime/Projects/Public/PluginDescriptor.h
--- a/Engine/Source/Runtime/Projects/Public/PluginDescriptor.h
+++ b/Engine/Source/Runtime/Projects/Public/PluginDescriptor.h
@@ -84,6 +84,9 @@ struct PROJECTS_API FPluginDescriptor
 	/** The Verse path to the root of this plugin's content directory */
 	FString VersePath;
 
+	/** If Verse code associated with this plugin is allowed to be seen by users outside of Epic */
+	bool bVersePublicApi;
+
 	/** Whether this plugin should be enabled by default for all projects */
 	EPluginEnabledByDefault EnabledByDefault;
 
