Fixed misspelt variable name.

#rnx

diff --git a/Engine/Source/Runtime/Projects/Private/PluginReferenceDescriptor.cpp b/Engine/Source/Runtime/Projects/Private/PluginReferenceDescriptor.cpp
--- a/Engine/Source/Runtime/Projects/Private/PluginReferenceDescriptor.cpp
+++ b/Engine/Source/Runtime/Projects/Private/PluginReferenceDescriptor.cpp
@@ -131,7 +131,7 @@ bool FPluginReferenceDescriptor::Read(const FJsonObject& Object, FText* OutFailR
 	AdditionalFieldsToWrite.Reset();
 #endif // WITH_EDITOR
 
-	bool bSucces = true;
+	bool bSuccess = true;
 
 	// Get the name
 	if (!Object.TryGetStringField(TEXT("Name"), Name))
@@ -140,7 +140,7 @@ bool FPluginReferenceDescriptor::Read(const FJsonObject& Object, FText* OutFailR
 		{
 			*OutFailReason = LOCTEXT("PluginReferenceWithoutName", "Plugin references must have a 'Name' field");
 		}
-		bSucces = false;
+		bSuccess = false;
 	}
 
 	// Get the enabled field
@@ -150,7 +150,7 @@ bool FPluginReferenceDescriptor::Read(const FJsonObject& Object, FText* OutFailR
 		{
 			*OutFailReason = LOCTEXT("PluginReferenceWithoutEnabled", "Plugin references must have an 'Enabled' field");
 		}
-		bSucces = false;
+		bSuccess = false;
 	}
 
 	// Read the optional field
@@ -176,7 +176,7 @@ bool FPluginReferenceDescriptor::Read(const FJsonObject& Object, FText* OutFailR
 	Object.TryGetStringArrayField(TEXT("SupportedTargetPlatforms"), SupportedTargetPlatforms);
 	Object.TryGetBoolField(TEXT("HasExplicitPlatforms"), bHasExplicitPlatforms);
 
-	return bSucces;
+	return bSuccess;
 }
 
 bool FPluginReferenceDescriptor::Read(const FJsonObject& Object, FText& OutFailReason, TSharedPtr<FJsonObject> ObjectPtr /*= nullptr*/)
