Adding the ability to supply multiple versions of a named plugin to UE (defaulting to the latest), along with the ability to reference a specific version.

diff --git a/Engine/Source/Runtime/Projects/Private/PluginManager.cpp b/Engine/Source/Runtime/Projects/Private/PluginManager.cpp
--- a/Engine/Source/Runtime/Projects/Private/PluginManager.cpp
+++ b/Engine/Source/Runtime/Projects/Private/PluginManager.cpp
@@ -89,6 +89,152 @@ namespace PluginSystemDefs
 	}
 }
 
+/** 
+ * Set of simple (ideally inlinable) helper methods intended to obscure how 
+ * `FDiscoveredPluginMap` is implemented (so that we can more easily change out it's type as needed).
+ * 
+ * Also dictactes how we separate the one "offered"  plugin out from other versions of the same plugin 
+ * (see DiscoveredPluginMapUtils::EInsertionType).
+ */
+namespace DiscoveredPluginMapUtils
+{
+	using FDiscoveredPluginMap = TMap<FString, TArray<TSharedRef<FPlugin>>>; // intended to match FPluginManager::FDiscoveredPluginMap
+
+	/** 
+	 * Given how UE is currently structured, we can only load/enable/display one plugin for any given name.
+	 * This enum is used to distinctualize how a plugin is handled -- should it be promoted as the one "offered" plugin, or supressed?
+	 */
+	enum class EInsertionType
+	{
+		AsOfferedPlugin,   // Denotes the one plugin you want visible to the rest of the engine.
+		AsSuppressedPlugin // Denotes discovered plugins that aren't reported to the rest of the engine (they're just for internal tracking and consideration).
+	};
+
+	/** Returns the one "offered" plugin from the specified entry. */
+	static const TSharedRef<FPlugin>& ResolvePluginFromMapVal(const FDiscoveredPluginMap::ValueType& PluginMapValue)
+	{
+		// The "top" plugin in the list is the one we choose to present to the rest of UE
+		return PluginMapValue.Top();
+	}
+	static TSharedRef<FPlugin>& ResolvePluginFromMapVal(FDiscoveredPluginMap::ValueType& PluginMapValue)
+	{
+		// The "top" plugin in the list is the one we choose to present to the rest of UE
+		return PluginMapValue.Top();
+	}
+
+	/** Returns the one "offered" plugin for the specified plugin name (if any exists). */
+	static TSharedRef<FPlugin>* FindPluginInMap(FDiscoveredPluginMap& InMap, const FDiscoveredPluginMap::KeyType& Name)
+	{
+		FDiscoveredPluginMap::ValueType* PluginList = InMap.Find(Name);
+		if (PluginList && !PluginList->IsEmpty())
+		{
+			// The "top" plugin in the list is the one we choose to present to the rest of UE
+			return &PluginList->Top();
+		}
+		return nullptr;
+	}
+
+	template<typename KeyComparable>
+	static TSharedRef<FPlugin>* FindPluginInMap_ByHash(FDiscoveredPluginMap& InMap, const uint32 KeyHash, const KeyComparable& Key)
+	{
+		FDiscoveredPluginMap::ValueType* PluginList = InMap.FindByHash(KeyHash, Key);
+		if (PluginList && !PluginList->IsEmpty())
+		{
+			// The "top" plugin in the list is the one we choose to present to the rest of UE
+			return &PluginList->Top();
+		}
+		return nullptr;
+	}
+
+	static TSharedRef<FPlugin>& FindPluginInMap_Checked(FDiscoveredPluginMap& InMap, const FDiscoveredPluginMap::KeyType& Name)
+	{
+		TSharedRef<FPlugin>* PluginPtr = FindPluginInMap(InMap, Name);
+		check(PluginPtr != nullptr);
+		return *PluginPtr;
+	}
+
+	/** Returns an array of all the discovered plugins with a matching name (null if none were discovered). */
+	static TArray<TSharedRef<FPlugin>>* FindAllPluginVersionsWithName(FDiscoveredPluginMap& InMap, const FDiscoveredPluginMap::KeyType& Name)
+	{
+		return InMap.Find(Name);
+	}
+
+	/** Internal insersion methods used for controlling the one "offered" plugin we share with the rest of the engine. */
+	static void InsertPlugin_AsOffered(FDiscoveredPluginMap::ValueType& MapEntry, const TSharedRef<FPlugin>& NewPlugin)
+	{
+		MapEntry.Push(NewPlugin);
+	}
+	static void InsertPlugin_AsSuppressed(FDiscoveredPluginMap::ValueType& MapEntry, const TSharedRef<FPlugin>& NewPlugin)
+	{
+		MapEntry.Insert(NewPlugin, FMath::Max(MapEntry.Num() - 1, 0));
+	}
+
+	/* 
+	 * Adds the specified plugin to the map, blindly assuming that the specified plugin is not a duplicate already in the map.
+	 * NOTE: When specifying `EInsertionType::AsOfferedPlugin`, you're overridding the previous "offered" plugin of this name.
+	 */
+	static void InsertPluginIntoMap(FDiscoveredPluginMap& Map, const FDiscoveredPluginMap::KeyType& Key, const TSharedRef<FPlugin>& NewPlugin, const EInsertionType InsertionType = EInsertionType::AsOfferedPlugin)
+	{
+		FDiscoveredPluginMap::ValueType& PluginList = Map.FindOrAdd(Key);
+		switch (InsertionType)
+		{
+		case EInsertionType::AsOfferedPlugin:
+			InsertPlugin_AsOffered(PluginList, NewPlugin);
+			break;
+
+		case EInsertionType::AsSuppressedPlugin:
+			InsertPlugin_AsSuppressed(PluginList, NewPlugin);
+			break;
+		}
+	}
+	static void InsertPluginIntoMap_ByHash(FDiscoveredPluginMap& Map, const uint32 KeyHash, const FDiscoveredPluginMap::KeyType& Key, const TSharedRef<FPlugin>& NewPlugin, const EInsertionType InsertionType = EInsertionType::AsOfferedPlugin)
+	{
+		FDiscoveredPluginMap::ValueType& PluginList = Map.FindOrAddByHash(KeyHash, Key);
+		switch (InsertionType)
+		{
+		case EInsertionType::AsOfferedPlugin:
+			InsertPlugin_AsOffered(PluginList, NewPlugin);
+			break;
+
+		case EInsertionType::AsSuppressedPlugin:
+			InsertPlugin_AsSuppressed(PluginList, NewPlugin);
+			break;
+		}
+	}
+
+	/**
+	 * Given how UE is currently structured, we can only load/enable/display one plugin for a given name -- this is the "offered" plugin.
+	 * This method let's us swap out the current "offered" plugin with the one we specify.
+	 * 
+	 * NOTE: This method assumes that the specified plugin pointer is already in the map under the specified name.
+	 * NOTE: After this method is ran the supplied TSharedRef<> is subject to change (use the returned value instead).
+	 * 
+	 * @return Null if the promotion failed, otherwise returns a new TSharedRef<> referencing the plugin that `PluginPtr` was initially.
+	 */
+	static const TSharedRef<FPlugin>* PromotePluginToOfferedVersion(FDiscoveredPluginMap& Map, TSharedRef<FPlugin>* PluginPtr)
+	{
+		const FString& Name = (*PluginPtr)->GetName();
+		FDiscoveredPluginMap::ValueType* PluginList = Map.Find(Name);
+
+		// Ensure the specified plugin is in the array
+		if (ensureAlwaysMsgf(PluginList && PluginList->GetData() >= PluginPtr && PluginPtr < PluginList->GetData() + PluginList->Num(),
+			TEXT("Specified plugin is not registered under the name '%s'. Failed to promote it."), *Name))
+		{
+			TSharedRef<FPlugin>* TopPluginPtr = &PluginList->Top();
+			if (TopPluginPtr != PluginPtr)
+			{
+				// Swap pointers in the array (promoting the specified plugin to the "top").
+				TSharedRef<FPlugin> OldTopPluginRef = *TopPluginPtr;
+				*TopPluginPtr = *PluginPtr;
+				*PluginPtr = OldTopPluginRef;
+			}
+
+			return TopPluginPtr;
+		}
+		return nullptr;
+	}
+}
+
 FPlugin::FPlugin(const FString& InFileName, const FPluginDescriptor& InDescriptor, EPluginType InType)
 	: Name(FPaths::GetBaseFilename(InFileName))
 	, FileName(InFileName)
@@ -224,14 +370,14 @@ FPluginManager::~FPluginManager()
 void FPluginManager::RefreshPluginsList()
 {
 	// Read a new list of all plugins
-	TMap<FString, TSharedRef<FPlugin>> NewPlugins;
+	FDiscoveredPluginMap NewPlugins;
 	ReadAllPlugins(NewPlugins, PluginDiscoveryPaths);
 
 	// Build a list of filenames for plugins which are enabled, and remove the rest
 	TSet<FString> EnabledPluginFileNames;
-	for(TMap<FString, TSharedRef<FPlugin>>::TIterator Iter(AllPlugins); Iter; ++Iter)
+	for(FDiscoveredPluginMap::TIterator Iter(AllPlugins); Iter; ++Iter)
 	{
-		const TSharedRef<FPlugin>& Plugin = Iter.Value();
+		const TSharedRef<FPlugin>& Plugin = DiscoveredPluginMapUtils::ResolvePluginFromMapVal(Iter.Value());
 		if(Plugin->bEnabled)
 		{
 			EnabledPluginFileNames.Add(Plugin->FileName);
@@ -243,22 +389,22 @@ void FPluginManager::RefreshPluginsList()
 	}
 
 	// Add all the plugins which aren't already enabled
-	for(const TPair<FString, TSharedRef<FPlugin>>& NewPluginPair: NewPlugins)
+	for(const FDiscoveredPluginMap::ElementType& NewPluginPair: NewPlugins)
 	{
-		const TSharedRef<FPlugin>& NewPlugin = NewPluginPair.Value;
+		const TSharedRef<FPlugin>& NewPlugin = DiscoveredPluginMapUtils::ResolvePluginFromMapVal(NewPluginPair.Value);
 		if(!EnabledPluginFileNames.Contains(NewPlugin->FileName))
 		{
 			const uint32 PluginNameHash = GetTypeHash(NewPlugin->GetName());
-			AllPlugins.AddByHash(PluginNameHash, NewPlugin->GetName(), NewPlugin);
+			DiscoveredPluginMapUtils::InsertPluginIntoMap_ByHash(AllPlugins, PluginNameHash, NewPlugin->GetName(), NewPlugin);
 			PluginsToConfigure.AddByHash(PluginNameHash, NewPlugin->GetName());
 		}
 	}
 
 #if WITH_EDITOR
 	ModuleNameToPluginMap.Reset();
-	for (const TPair<FString, TSharedRef<FPlugin>>& PluginPair : AllPlugins)
+	for (const FDiscoveredPluginMap::ElementType& PluginPair : AllPlugins)
 	{
-		AddToModuleNameToPluginMap(PluginPair.Value);
+		AddToModuleNameToPluginMap(DiscoveredPluginMapUtils::ResolvePluginFromMapVal(PluginPair.Value));
 	}
 #endif //if WITH_EDITOR
 }
@@ -315,17 +461,17 @@ bool FPluginManager::AddToPluginsList(const FString& PluginFilename, FText* OutF
 		}
 
 		// Create the plugin
-		TMap<FString, TSharedRef<FPlugin>> NewPlugins;
+		FDiscoveredPluginMap NewPlugins;
 		TArray<TSharedRef<FPlugin>> ChildPlugins;
 		CreatePluginObject(PluginFilename, Descriptor, PluginType, NewPlugins, ChildPlugins);
 		ensureMsgf(ChildPlugins.Num() == 0, TEXT("AddToPluginsList does not allow plugins with bIsPluginExtension set to true. Plugin: %s"), *PluginFilename);
 		ensure(NewPlugins.Num() == 1);
 		
 		// Add the loaded plugin
-		TSharedRef<FPlugin>* NewPlugin = NewPlugins.Find(PluginName);
+		TSharedRef<FPlugin>* NewPlugin = DiscoveredPluginMapUtils::FindPluginInMap(NewPlugins, PluginName);
 		if (ensure(NewPlugin))
 		{
-			AllPlugins.Add(PluginName, *NewPlugin);
+			DiscoveredPluginMapUtils::InsertPluginIntoMap(AllPlugins, PluginName, *NewPlugin);
 
 #if WITH_EDITOR
 			AddToModuleNameToPluginMap(*NewPlugin);
@@ -353,7 +499,7 @@ bool FPluginManager::RemoveFromPluginsList(const FString& PluginFilename, FText*
 #if (WITH_ENGINE && !IS_PROGRAM) || WITH_PLUGIN_SUPPORT
 	FString PluginName = FPaths::GetBaseFilename(PluginFilename);
 
-	TSharedRef<FPlugin>* MaybePlugin = AllPlugins.Find(PluginName);
+	TSharedRef<FPlugin>* MaybePlugin = DiscoveredPluginMapUtils::FindPluginInMap(AllPlugins, PluginName);
 	if (MaybePlugin == nullptr)
 	{
 		return true;
@@ -401,17 +547,17 @@ void FPluginManager::DiscoverAllPlugins()
 	ReadAllPlugins(AllPlugins, PluginDiscoveryPaths);
 
 	PluginsToConfigure.Reserve(AllPlugins.Num());
-	for (const TPair<FString, TSharedRef<FPlugin>>& PluginPair : AllPlugins)
+	for (const FDiscoveredPluginMap::ElementType& PluginPair : AllPlugins)
 	{
 		PluginsToConfigure.Add(PluginPair.Key);
 #if WITH_EDITOR
 		BuiltInPluginNames.Add(PluginPair.Key);
-		AddToModuleNameToPluginMap(PluginPair.Value);
+		AddToModuleNameToPluginMap(DiscoveredPluginMapUtils::ResolvePluginFromMapVal(PluginPair.Value));
 #endif //if WITH_EDITOR
 	}
 }
 
-void FPluginManager::ReadAllPlugins(TMap<FString, TSharedRef<FPlugin>>& Plugins, const TSet<FString>& ExtraSearchPaths)
+void FPluginManager::ReadAllPlugins(FDiscoveredPluginMap& Plugins, const TSet<FString>& ExtraSearchPaths)
 {
 #if (WITH_ENGINE && !IS_PROGRAM) || (WITH_PLUGIN_SUPPORT && !UE_DISABLE_PLUGIN_DISCOVERY)
 	const FProjectDescriptor* Project = IProjectManager::Get().GetCurrentProject();
@@ -530,7 +676,7 @@ void FPluginManager::ReadAllPlugins(TMap<FString, TSharedRef<FPlugin>>& Plugins,
 		if (Tokens.Num() == 2)
 		{
 			FString ParentPluginName = Tokens[0];
-			ParentPtr = Plugins.Find(ParentPluginName);
+			ParentPtr = DiscoveredPluginMapUtils::FindPluginInMap(Plugins, ParentPluginName);
 		}
 		if (ParentPtr != nullptr)
 		{
@@ -571,7 +717,7 @@ void FPluginManager::ReadAllPlugins(TMap<FString, TSharedRef<FPlugin>>& Plugins,
 #endif
 }
 
-void FPluginManager::ReadPluginsInDirectory(const FString& PluginsDirectory, const EPluginType Type, TMap<FString, TSharedRef<FPlugin>>& Plugins, TArray<TSharedRef<FPlugin>>& ChildPlugins)
+void FPluginManager::ReadPluginsInDirectory(const FString& PluginsDirectory, const EPluginType Type, FDiscoveredPluginMap& Plugins, TArray<TSharedRef<FPlugin>>& ChildPlugins)
 {
 	// Make sure the directory even exists
 	if(FPlatformFileManager::Get().GetPlatformFile().DirectoryExists(*PluginsDirectory))
@@ -655,7 +801,7 @@ void FPluginManager::FindPluginManifestsInDirectory(const FString& PluginManifes
 	IFileManager::Get().IterateDirectory(*PluginManifestDirectory, Visitor);
 }
 
-void FPluginManager::CreatePluginObject(const FString& FileName, const FPluginDescriptor& Descriptor, const EPluginType Type, TMap<FString, TSharedRef<FPlugin>>& Plugins, TArray<TSharedRef<FPlugin>>& ChildPlugins)
+void FPluginManager::CreatePluginObject(const FString& FileName, const FPluginDescriptor& Descriptor, const EPluginType Type, FDiscoveredPluginMap& Plugins, TArray<TSharedRef<FPlugin>>& ChildPlugins)
 {
 	TSharedRef<FPlugin> Plugin = MakeShareable(new FPlugin(FileName, Descriptor, Type));
 
@@ -669,20 +815,55 @@ void FPluginManager::CreatePluginObject(const FString& FileName, const FPluginDe
 	FString FullPath = FPaths::ConvertRelativePathToFull(FileName);
 	UE_LOG(LogPluginManager, Verbose, TEXT("Read plugin descriptor for %s, from %s"), *Plugin->GetName(), *FullPath);
 
-	const TSharedRef<FPlugin>* ExistingPlugin = Plugins.Find(Plugin->GetName());
+	TSharedRef<FPlugin>* ExistingPlugin = DiscoveredPluginMapUtils::FindPluginInMap(Plugins, Plugin->GetName());
 	if (ExistingPlugin == nullptr)
 	{
-		Plugins.Add(Plugin->GetName(), Plugin);
+		DiscoveredPluginMapUtils::InsertPluginIntoMap(Plugins, Plugin->GetName(), Plugin);
 	}
+	// We allow for duplicates of plugins between engine and the project, favoring the project level plugin
 	else if ((*ExistingPlugin)->Type == EPluginType::Engine && Type == EPluginType::Project)
 	{
-		Plugins[Plugin->GetName()] = Plugin;
-		UE_LOG(LogPluginManager, Verbose, TEXT("Replacing engine version of '%s' plugin with game version"), *Plugin->GetName());
+		UE_LOG(LogPluginManager, Verbose, TEXT("By default, prioritizing project plugin (%s) over the corresponding engine version (%s)."), *Plugin->FileName, *(*ExistingPlugin)->FileName);
+		DiscoveredPluginMapUtils::InsertPluginIntoMap(Plugins, Plugin->GetName(), Plugin, DiscoveredPluginMapUtils::EInsertionType::AsOfferedPlugin);
 	}
-	else if (((*ExistingPlugin)->Type != EPluginType::Project || Type != EPluginType::Engine) && (*ExistingPlugin)->FileName != Plugin->FileName)
+	else if ((*ExistingPlugin)->FileName != Plugin->FileName)
 	{
-		UE_LOG(LogPluginManager, Warning, TEXT("Plugin '%s' exists at '%s' and '%s' - second location will be ignored"), *Plugin->GetName(), *(*ExistingPlugin)->FileName, *Plugin->FileName);
+		if ((*ExistingPlugin)->Type == EPluginType::Project && Type == EPluginType::Engine)
+		{
+			// Project plugins are favored over engine plugins, so we don't want to warn in this case
+			// (instead we mimic the Verbose log from above, and just explain which plugin we're favoring)
+			UE_LOG(LogPluginManager, Verbose, TEXT("By default, prioritizing project plugin (%s) over the corresponding engine version (%s)."), *(*ExistingPlugin)->FileName, *Plugin->FileName);
+			DiscoveredPluginMapUtils::InsertPluginIntoMap(Plugins, Plugin->GetName(), Plugin, DiscoveredPluginMapUtils::EInsertionType::AsSuppressedPlugin);
+		}
+		else
+		{
+			const int32 ExistingVersion = (*ExistingPlugin)->GetDescriptor().Version;
+			if (ExistingVersion != Descriptor.Version)
+			{
+				UE_LOG(
+					LogPluginManager,
+					Verbose,
+					TEXT("By default, prioritizing newer version (v%d) of '%s' plugin, over older version (v%d)."),
+					FMath::Max(ExistingVersion, Descriptor.Version),
+					*Plugin->GetName(),
+					FMath::Min(ExistingVersion, Descriptor.Version));
+
+				if (ExistingVersion < Descriptor.Version)
+				{
+					DiscoveredPluginMapUtils::InsertPluginIntoMap(Plugins, Plugin->GetName(), Plugin, DiscoveredPluginMapUtils::EInsertionType::AsOfferedPlugin);
+				}
+				else
+				{
+					DiscoveredPluginMapUtils::InsertPluginIntoMap(Plugins, Plugin->GetName(), Plugin, DiscoveredPluginMapUtils::EInsertionType::AsSuppressedPlugin);
+				}
+			}
+			else
+			{
+				UE_LOG(LogPluginManager, Warning, TEXT("The same version (v%d) of plugin '%s' exists at '%s' and '%s' - second location will be ignored."), ExistingVersion, *Plugin->GetName(), *(*ExistingPlugin)->FileName, *Plugin->FileName);
+			}
+		}
 	}
+	// else, same plugin? Do nothing -- keep from adding it twice.
 }
 
 // Helper class to find all pak files.
@@ -984,7 +1165,7 @@ bool FPluginManager::ConfigureEnabledPlugins()
 
 				for (const FString& PluginName : PluginsToConfigure)
 				{
-					const TSharedRef<FPlugin>& Plugin = AllPlugins.FindChecked(PluginName);
+					const TSharedRef<FPlugin>& Plugin = DiscoveredPluginMapUtils::FindPluginInMap_Checked(AllPlugins, PluginName);
 					if (Plugin->IsEnabledByDefault(bAllowEnginePluginsEnabledByDefault) && !ConfiguredPluginNames.Contains(PluginName))
 					{
 						if (!ConfigureEnabledPluginForCurrentTarget(FPluginReferenceDescriptor(PluginName, true), EnabledPlugins))
@@ -1096,7 +1277,7 @@ bool FPluginManager::ConfigureEnabledPlugins()
 			TArray<TSharedRef<FPlugin>> PluginsArray;
 			for (const FString& PluginName : PluginsToConfigure)
 			{
-				TSharedRef<FPlugin> Plugin = AllPlugins.FindChecked(PluginName);
+				TSharedRef<FPlugin> Plugin = DiscoveredPluginMapUtils::FindPluginInMap_Checked(AllPlugins, PluginName);
 				// check all plugins that were not in a BinaryConfig
 				if (!BinaryConfigPlugins.Contains(PluginName))
 				{
@@ -1321,7 +1502,7 @@ bool FPluginManager::RequiresTempTargetForCodePlugin(const FProjectDescriptor* P
 	return false;
 }
 
-bool FPluginManager::GetCodePluginsForProject(const FProjectDescriptor* ProjectDescriptor, const FString& Platform, EBuildConfiguration Configuration, EBuildTargetType TargetType, const TMap<FString, TSharedRef<FPlugin>>& AllPlugins, TSet<FString>& CodePluginNames, const FPluginReferenceDescriptor*& OutMissingPlugin)
+bool FPluginManager::GetCodePluginsForProject(const FProjectDescriptor* ProjectDescriptor, const FString& Platform, EBuildConfiguration Configuration, EBuildTargetType TargetType, FDiscoveredPluginMap& AllPlugins, TSet<FString>& CodePluginNames, const FPluginReferenceDescriptor*& OutMissingPlugin)
 {
 	// Can only check the current project at the moment, since we won't have enumerated them otherwise
 	check(ProjectDescriptor == nullptr || ProjectDescriptor == IProjectManager::Get().GetCurrentProject());
@@ -1357,9 +1538,9 @@ bool FPluginManager::GetCodePluginsForProject(const FProjectDescriptor* ProjectD
 	}
 
 	// Add the plugins which are enabled by default
-	for (const TPair<FString, TSharedRef<FPlugin>>& PluginPair : AllPlugins)
+	for (const FDiscoveredPluginMap::ElementType& PluginPair : AllPlugins)
 	{
-		if (PluginPair.Value->IsEnabledByDefault(bAllowEnginePluginsEnabledByDefault) && !ConfiguredPluginNames.Contains(PluginPair.Key))
+		if (DiscoveredPluginMapUtils::ResolvePluginFromMapVal(PluginPair.Value)->IsEnabledByDefault(bAllowEnginePluginsEnabledByDefault) && !ConfiguredPluginNames.Contains(PluginPair.Key))
 		{
 			if (!ConfigureEnabledPluginForTarget(FPluginReferenceDescriptor(PluginPair.Key, true), ProjectDescriptor, FString(), Platform, Configuration, TargetType, bLoadPluginsForTargetPlatforms, AllPlugins, EnabledPlugins, OutMissingPlugin))
 			{
@@ -1424,7 +1605,7 @@ bool FPluginManager::ConfigureEnabledPluginForCurrentTarget(const FPluginReferen
 	return true;
 }
 
-bool FPluginManager::ConfigureEnabledPluginForTarget(const FPluginReferenceDescriptor& FirstReference, const FProjectDescriptor* ProjectDescriptor, const FString& TargetName, const FString& Platform, EBuildConfiguration Configuration, EBuildTargetType TargetType, bool bLoadPluginsForTargetPlatforms, const TMap<FString, TSharedRef<FPlugin>>& AllPlugins, TMap<FString, FPlugin*>& EnabledPlugins, const FPluginReferenceDescriptor*& OutMissingPlugin)
+bool FPluginManager::ConfigureEnabledPluginForTarget(const FPluginReferenceDescriptor& FirstReference, const FProjectDescriptor* ProjectDescriptor, const FString& TargetName, const FString& Platform, EBuildConfiguration Configuration, EBuildTargetType TargetType, bool bLoadPluginsForTargetPlatforms, FDiscoveredPluginMap& AllPlugins, TMap<FString, FPlugin*>& EnabledPlugins, const FPluginReferenceDescriptor*& OutMissingPlugin)
 {
 	if (!EnabledPlugins.Contains(FirstReference.Name))
 	{
@@ -1457,7 +1638,32 @@ bool FPluginManager::ConfigureEnabledPluginForTarget(const FPluginReferenceDescr
 			}
 
 			// Find the plugin being enabled
-			const TSharedRef<FPlugin>* PluginPtr = AllPlugins.Find(Reference.Name);
+			const TSharedRef<FPlugin>* PluginPtr = nullptr;
+			if (Reference.RequestedVersion.IsSet())
+			{
+				if (TArray<TSharedRef<FPlugin>>* PluginVersions = DiscoveredPluginMapUtils::FindAllPluginVersionsWithName(AllPlugins, Reference.Name))
+				{
+					for (TSharedRef<FPlugin>& PluginVersion : *PluginVersions)
+					{
+						if (PluginVersion->GetDescriptor().Version == Reference.RequestedVersion.GetValue())
+						{
+							PluginPtr = DiscoveredPluginMapUtils::PromotePluginToOfferedVersion(AllPlugins, &PluginVersion);
+							break;
+						}
+					}
+
+					if (!PluginPtr)
+					{
+						UE_LOG(LogPluginManager, Warning, TEXT("Failed to find specific version (v%d) of plugin '%s'. Other versions exist, but the explicit version requested was missing."),
+							FirstReference.RequestedVersion.GetValue(), *Reference.Name);
+					}
+				}
+			}
+			else
+			{
+				PluginPtr = DiscoveredPluginMapUtils::FindPluginInMap(AllPlugins, Reference.Name);
+			}
+
 			if (PluginPtr == nullptr)
 			{
 				// Ignore any optional plugins
@@ -1539,6 +1745,18 @@ bool FPluginManager::ConfigureEnabledPluginForTarget(const FPluginReferenceDescr
 			EnabledPlugins.Add(Plugin.GetName(), &Plugin);
 		}
 	}
+	else if (FirstReference.RequestedVersion.IsSet())
+	{
+		UE_CLOG(
+			(*EnabledPlugins.Find(FirstReference.Name))->GetDescriptor().Version != FirstReference.RequestedVersion.GetValue(),
+			LogPluginManager, 
+			Error, 
+			TEXT("Requested explicit version (v%d) of plugin '%s', but a different version (v%d) was already enabled by another source."),
+			FirstReference.RequestedVersion.GetValue(),
+			*FirstReference.Name,
+			(*EnabledPlugins.Find(FirstReference.Name))->GetDescriptor().Version
+		);
+	}
 	return true;
 }
 
@@ -1640,7 +1858,7 @@ bool FPluginManager::PromptToLoadIncompatiblePlugin(const FPlugin& Plugin)
 
 TSharedPtr<FPlugin> FPluginManager::FindPluginInstance(const FString& Name)
 {
-	const TSharedRef<FPlugin>* Instance = AllPlugins.Find(Name);
+	const TSharedRef<FPlugin>* Instance = DiscoveredPluginMapUtils::FindPluginInMap(AllPlugins, Name);
 	if (Instance == nullptr)
 	{
 		return TSharedPtr<FPlugin>();
@@ -1724,9 +1942,9 @@ bool FPluginManager::LoadModulesForEnabledPlugins( const ELoadingPhase::Type Loa
 		FScopedSlowTask SlowTask(AllPlugins.Num());
 
 		// Load plugins!
-		for (const TPair<FString, TSharedRef< FPlugin >>& PluginPair : AllPlugins)
+		for (const FDiscoveredPluginMap::ElementType& PluginPair : AllPlugins)
 		{
-			const TSharedRef<FPlugin>& Plugin = PluginPair.Value;
+			const TSharedRef<FPlugin>& Plugin = DiscoveredPluginMapUtils::ResolvePluginFromMapVal(PluginPair.Value);
 
 			SlowTask.EnterProgressFrame(1);
 
@@ -1771,9 +1989,9 @@ void FPluginManager::GetLocalizationPathsForEnabledPlugins( TArray<FString>& Out
 	}
 
 	// Gather the paths from all plugins that have localization targets that are loaded based on the current runtime environment
-	for (const TPair<FString, TSharedRef<FPlugin>>& PluginPair : AllPlugins)
+	for (const FDiscoveredPluginMap::ElementType& PluginPair : AllPlugins)
 	{
-		const TSharedRef<FPlugin>& Plugin = PluginPair.Value;
+		const TSharedRef<FPlugin>& Plugin = DiscoveredPluginMapUtils::ResolvePluginFromMapVal(PluginPair.Value);
 		if (!Plugin->bEnabled || Plugin->GetDescriptor().LocalizationTargets.Num() == 0)
 		{
 			continue;
@@ -1819,9 +2037,9 @@ bool FPluginManager::CheckModuleCompatibility(TArray<FString>& OutIncompatibleMo
 	}
 
 	bool bResult = true;
-	for(const TPair<FString, TSharedRef<FPlugin>>& PluginPair : AllPlugins)
+	for(const FDiscoveredPluginMap::ElementType& PluginPair : AllPlugins)
 	{
-		const TSharedRef< FPlugin > &Plugin = PluginPair.Value;
+		const TSharedRef< FPlugin > &Plugin = DiscoveredPluginMapUtils::ResolvePluginFromMapVal(PluginPair.Value);
 
 		TArray<FString> IncompatibleModules;
 		if (Plugin->bEnabled && !FModuleDescriptor::CheckModuleCompatibility(Plugin->Descriptor.Modules, IncompatibleModules))
@@ -1848,7 +2066,7 @@ IPluginManager& IPluginManager::Get()
 TSharedPtr<IPlugin> FPluginManager::FindPlugin(const FStringView Name)
 {
 	const uint32 NameHash = GetTypeHash(Name);
-	const TSharedRef<FPlugin>* Instance = AllPlugins.FindByHash(NameHash, Name);
+	const TSharedRef<FPlugin>* Instance = DiscoveredPluginMapUtils::FindPluginInMap_ByHash(AllPlugins, NameHash, Name);
 	if (Instance == nullptr)
 	{
 		return TSharedPtr<IPlugin>();
@@ -1882,9 +2100,9 @@ TArray<TSharedRef<IPlugin>> FPluginManager::GetEnabledPlugins()
 {
 	TArray<TSharedRef<IPlugin>> Plugins;
 	Plugins.Reserve(AllPlugins.Num());
-	for(TPair<FString, TSharedRef<FPlugin>>& PluginPair : AllPlugins)
+	for(FDiscoveredPluginMap::ElementType& PluginPair : AllPlugins)
 	{
-		TSharedRef<FPlugin>& PossiblePlugin = PluginPair.Value;
+		TSharedRef<FPlugin>& PossiblePlugin = DiscoveredPluginMapUtils::ResolvePluginFromMapVal(PluginPair.Value);
 		if(PossiblePlugin->bEnabled)
 		{
 			Plugins.Add(PossiblePlugin);
@@ -1897,9 +2115,9 @@ TArray<TSharedRef<IPlugin>> FPluginManager::GetEnabledPluginsWithContent() const
 {
 	TArray<TSharedRef<IPlugin>> Plugins;
 	Plugins.Reserve(AllPlugins.Num());
-	for (const TPair<FString, TSharedRef<FPlugin>>& PluginPair : AllPlugins)
+	for (const FDiscoveredPluginMap::ElementType& PluginPair : AllPlugins)
 	{
-		const TSharedRef<FPlugin>& PluginRef = PluginPair.Value;
+		const TSharedRef<FPlugin>& PluginRef = DiscoveredPluginMapUtils::ResolvePluginFromMapVal(PluginPair.Value);
 		const FPlugin& Plugin = *PluginRef;
 		if (Plugin.IsEnabled() && Plugin.CanContainContent())
 		{
@@ -1913,9 +2131,9 @@ TArray<TSharedRef<IPlugin>> FPluginManager::GetEnabledPluginsWithVerse() const
 {
 	TArray<TSharedRef<IPlugin>> Plugins;
 	Plugins.Reserve(AllPlugins.Num());
-	for (const TPair<FString, TSharedRef<FPlugin>>& PluginPair : AllPlugins)
+	for (const FDiscoveredPluginMap::ElementType& PluginPair : AllPlugins)
 	{
-		const TSharedRef<FPlugin>& PluginRef = PluginPair.Value;
+		const TSharedRef<FPlugin>& PluginRef = DiscoveredPluginMapUtils::ResolvePluginFromMapVal(PluginPair.Value);
 		const FPlugin& Plugin = *PluginRef;
 		if (Plugin.IsEnabled() && Plugin.CanContainVerse())
 		{
@@ -1929,9 +2147,9 @@ TArray<TSharedRef<IPlugin>> FPluginManager::GetEnabledPluginsWithContentOrVerse(
 {
 	TArray<TSharedRef<IPlugin>> Plugins;
 	Plugins.Reserve(AllPlugins.Num());
-	for (const TPair<FString, TSharedRef<FPlugin>>& PluginPair : AllPlugins)
+	for (const FDiscoveredPluginMap::ElementType& PluginPair : AllPlugins)
 	{
-		const TSharedRef<FPlugin>& PluginRef = PluginPair.Value;
+		const TSharedRef<FPlugin>& PluginRef = DiscoveredPluginMapUtils::ResolvePluginFromMapVal(PluginPair.Value);
 		const FPlugin& Plugin = *PluginRef;
 		if (Plugin.IsEnabled() && (Plugin.CanContainContent() || Plugin.CanContainVerse()))
 		{
@@ -1945,10 +2163,12 @@ TArray<TSharedRef<IPlugin>> FPluginManager::GetDiscoveredPlugins()
 {
 	TArray<TSharedRef<IPlugin>> Plugins;
 	Plugins.Reserve(AllPlugins.Num());
-	for (TPair<FString, TSharedRef<FPlugin>>& PluginPair : AllPlugins)
+
+	for (FDiscoveredPluginMap::ElementType& PluginPair : AllPlugins)
 	{
-		Plugins.Add(PluginPair.Value);
+		Plugins.Add(DiscoveredPluginMapUtils::ResolvePluginFromMapVal(PluginPair.Value));
 	}
+
 	return Plugins;
 }
 
@@ -2146,9 +2366,9 @@ bool FPluginManager::UnmountPluginFromExternalSource(const TSharedPtr<FPlugin>&
 FName FPluginManager::PackageNameFromModuleName(FName ModuleName)
 {
 	FName Result = ModuleName;
-	for (TMap<FString, TSharedRef<FPlugin>>::TIterator Iter(AllPlugins); Iter; ++Iter)
+	for (FDiscoveredPluginMap::TIterator Iter(AllPlugins); Iter; ++Iter)
 	{
-		const TSharedRef<FPlugin>& Plugin = Iter.Value();
+		const TSharedRef<FPlugin>& Plugin = DiscoveredPluginMapUtils::ResolvePluginFromMapVal(Iter.Value());
 		const TArray<FModuleDescriptor>& Modules = Plugin->Descriptor.Modules;
 		for (int Idx = 0; Idx < Modules.Num(); Idx++)
 		{
diff --git a/Engine/Source/Runtime/Projects/Private/PluginManager.h b/Engine/Source/Runtime/Projects/Private/PluginManager.h
--- a/Engine/Source/Runtime/Projects/Private/PluginManager.h
+++ b/Engine/Source/Runtime/Projects/Private/PluginManager.h
@@ -159,19 +159,20 @@ public:
 	virtual bool IntegratePluginsIntoConfig(FConfigCacheIni& ConfigSystem, const TCHAR* EngineIniName, const TCHAR* PlatformName, const TCHAR* StagedPluginsFile);
 
 private:
+	using FDiscoveredPluginMap = TMap<FString, TArray<TSharedRef<FPlugin>>>;
 
 	/** Searches for all plugins on disk and builds up the array of plugin objects.  Doesn't load any plugins. 
 	    This is called when the plugin manager singleton is first accessed. */
 	void DiscoverAllPlugins();
 
 	/** Reads all the plugin descriptors */
-	static void ReadAllPlugins(TMap<FString, TSharedRef<FPlugin>>& Plugins, const TSet<FString>& ExtraSearchPaths);
+	static void ReadAllPlugins(FDiscoveredPluginMap& Plugins, const TSet<FString>& ExtraSearchPaths);
 
 	/** Reads all the plugin descriptors from disk */
-	static void ReadPluginsInDirectory(const FString& PluginsDirectory, const EPluginType Type, TMap<FString, TSharedRef<FPlugin>>& Plugins, TArray<TSharedRef<FPlugin>>& ChildPlugins);
+	static void ReadPluginsInDirectory(const FString& PluginsDirectory, const EPluginType Type, FDiscoveredPluginMap& Plugins, TArray<TSharedRef<FPlugin>>& ChildPlugins);
 
 	/** Creates a FPlugin object and adds it to the given map */
-	static void CreatePluginObject(const FString& FileName, const FPluginDescriptor& Descriptor, const EPluginType Type, TMap<FString, TSharedRef<FPlugin>>& Plugins, TArray<TSharedRef<FPlugin>>& ChildPlugins);
+	static void CreatePluginObject(const FString& FileName, const FPluginDescriptor& Descriptor, const EPluginType Type, FDiscoveredPluginMap& Plugins, TArray<TSharedRef<FPlugin>>& ChildPlugins);
 
 	/** Finds all the plugin descriptors underneath a given directory */
 	static void FindPluginsInDirectory(const FString& PluginsDirectory, TArray<FString>& FileNames);
@@ -180,7 +181,7 @@ private:
 	static void FindPluginManifestsInDirectory(const FString& PluginManifestDirectory, TArray<FString>& FileNames);
 
 	/** Gets all the code plugins that are enabled for a content only project */
-	static bool GetCodePluginsForProject(const FProjectDescriptor* ProjectDescriptor, const FString& Platform, EBuildConfiguration Configuration, EBuildTargetType TargetType, const TMap<FString, TSharedRef<FPlugin>>& AllPlugins, TSet<FString>& CodePluginNames, const FPluginReferenceDescriptor*& OutMissingPlugin);
+	static bool GetCodePluginsForProject(const FProjectDescriptor* ProjectDescriptor, const FString& Platform, EBuildConfiguration Configuration, EBuildTargetType TargetType, FDiscoveredPluginMap& AllPlugins, TSet<FString>& CodePluginNames, const FPluginReferenceDescriptor*& OutMissingPlugin);
 
 	/** Sets the bPluginEnabled flag on all plugins found from DiscoverAllPlugins that are enabled in config */
 	bool ConfigureEnabledPlugins();
@@ -189,7 +190,7 @@ private:
 	bool ConfigureEnabledPluginForCurrentTarget(const FPluginReferenceDescriptor& FirstReference, TMap<FString, FPlugin*>& EnabledPlugins);
 
 	/** Adds a single enabled plugin and all its dependencies. */
-	static bool ConfigureEnabledPluginForTarget(const FPluginReferenceDescriptor& FirstReference, const FProjectDescriptor* ProjectDescriptor, const FString& TargetName, const FString& Platform, EBuildConfiguration Configuration, EBuildTargetType TargetType, bool bLoadPluginsForTargetPlatforms, const TMap<FString, TSharedRef<FPlugin>>& AllPlugins, TMap<FString, FPlugin*>& EnabledPlugins, const FPluginReferenceDescriptor*& OutMissingPlugin);
+	static bool ConfigureEnabledPluginForTarget(const FPluginReferenceDescriptor& FirstReference, const FProjectDescriptor* ProjectDescriptor, const FString& TargetName, const FString& Platform, EBuildConfiguration Configuration, EBuildTargetType TargetType, bool bLoadPluginsForTargetPlatforms, FDiscoveredPluginMap& AllPlugins, TMap<FString, FPlugin*>& EnabledPlugins, const FPluginReferenceDescriptor*& OutMissingPlugin);
 
 	/** Prompts the user to download a missing plugin from the given URL */
 	static bool PromptToDownloadPlugin(const FString& PluginName, const FString& MarketplaceURL);
@@ -228,7 +229,7 @@ private:
 
 private:
 	/** All of the plugins that we know about */
-	TMap< FString, TSharedRef< FPlugin > > AllPlugins;
+	FDiscoveredPluginMap AllPlugins;
 
 	/** Plugins that need to be configured to see if they should be enabled */
 	TSet<FString> PluginsToConfigure;
diff --git a/Engine/Source/Runtime/Projects/Private/PluginReferenceDescriptor.cpp b/Engine/Source/Runtime/Projects/Private/PluginReferenceDescriptor.cpp
--- a/Engine/Source/Runtime/Projects/Private/PluginReferenceDescriptor.cpp
+++ b/Engine/Source/Runtime/Projects/Private/PluginReferenceDescriptor.cpp
@@ -176,6 +176,20 @@ bool FPluginReferenceDescriptor::Read(const FJsonObject& Object, FText* OutFailR
 	Object.TryGetStringArrayField(TEXT("SupportedTargetPlatforms"), SupportedTargetPlatforms);
 	Object.TryGetBoolField(TEXT("HasExplicitPlatforms"), bHasExplicitPlatforms);
 
+	int32 ReadVersion;
+	if (Object.TryGetNumberField(TEXT("Version"), ReadVersion))
+	{
+		RequestedVersion = ReadVersion;
+		if (!bEnabled)
+		{
+			if (bSuccess && OutFailReason)
+			{
+				*OutFailReason = LOCTEXT("PluginReferenceDisabledWithVersion", "Plugin references cannot be used to disable explicit versions. Remove the 'Version' field when 'Enabled' is false.");
+			}
+			bSuccess = false;
+		}
+	}
+
 	return bSuccess;
 }
 
@@ -373,6 +387,15 @@ void FPluginReferenceDescriptor::UpdateJson(FJsonObject& JsonObject) const
 		JsonObject.RemoveField(TEXT("HasExplicitPlatforms"));
 	}
 
+	if (bEnabled && RequestedVersion.IsSet())
+	{
+		JsonObject.SetNumberField(TEXT("Version"), RequestedVersion.GetValue());
+	}
+	else
+	{
+		JsonObject.RemoveField(TEXT("Version"));
+	}
+
 	// Remove deprecated fields
 	JsonObject.RemoveField(TEXT("WhitelistPlatforms"));
 	JsonObject.RemoveField(TEXT("BlacklistPlatforms"));
diff --git a/Engine/Source/Runtime/Projects/Public/PluginReferenceDescriptor.h b/Engine/Source/Runtime/Projects/Public/PluginReferenceDescriptor.h
--- a/Engine/Source/Runtime/Projects/Public/PluginReferenceDescriptor.h
+++ b/Engine/Source/Runtime/Projects/Public/PluginReferenceDescriptor.h
@@ -53,6 +53,9 @@ struct PROJECTS_API FPluginReferenceDescriptor
 	/** When true, empty SupportedTargetPlatforms and PlatformAllowList are interpreted as 'no platforms' with the expectation that explicit platforms will be added in plugin platform extensions */
 	bool bHasExplicitPlatforms;
 
+	/** When set, specifies a specific version of the plugin that this references. */
+	TOptional<int32> RequestedVersion;
+
 #if WITH_EDITOR
 	/** Cached json for custom data */
 	mutable TSharedPtr<FJsonObject> CachedJson;
