Providing a way for external systems to extend the plugin configuration process.

diff --git a/Engine/Source/Runtime/Projects/Private/PluginManager.cpp b/Engine/Source/Runtime/Projects/Private/PluginManager.cpp
--- a/Engine/Source/Runtime/Projects/Private/PluginManager.cpp
+++ b/Engine/Source/Runtime/Projects/Private/PluginManager.cpp
@@ -23,6 +23,8 @@
 #include "PluginManifest.h"
 #include "HAL/PlatformTime.h"
 #include "Async/ParallelFor.h"
+#include "IPluginConfigServer.h"
+#include "Features/IModularFeatures.h"
 #if READ_TARGET_ENABLED_PLUGINS_FROM_RECEIPT
 #include "TargetReceipt.h"
 #endif
@@ -1133,6 +1135,37 @@ bool FPluginManager::ConfigureEnabledPlugins()
 			}
 #endif // READ_TARGET_ENABLED_PLUGINS_FROM_RECEIPT
 
+			auto ProcessPluginConfigurations = [&ConfiguredPluginNames, &EnabledPlugins, this](const TArray<FPluginReferenceDescriptor>& PluginReferences)->bool
+			{
+				for (const FPluginReferenceDescriptor& PluginReference : PluginReferences)
+				{
+					if (!ConfiguredPluginNames.Contains(PluginReference.Name))
+					{
+						if (!ConfigureEnabledPluginForCurrentTarget(PluginReference, EnabledPlugins))
+						{
+							return false;
+						}
+						ConfiguredPluginNames.Add(PluginReference.Name);
+					}
+				}
+				return true;
+			};
+
+			TArray<IPluginConfigServer*> PluginConfigServers = IModularFeatures::Get().GetModularFeatureImplementations<IPluginConfigServer>(IPluginConfigServer::GetModularFeatureName());
+			{
+				SCOPED_BOOT_TIMING("PluginConfigServers_PreProjConfig");
+				
+				for (const IPluginConfigServer* ConfigServer : PluginConfigServers)
+				{
+					TArray<FPluginReferenceDescriptor> PluginConfigs;
+					ConfigServer->PreProjConfig_GetPluginConfigurations(PluginConfigs);
+					if (!ProcessPluginConfigurations(PluginConfigs))
+					{
+						return false;
+					}					
+				}
+			}
+
 			bool bAllowEnginePluginsEnabledByDefault = true;
 			// Find all the plugin references in the project file
 			const FProjectDescriptor* ProjectDescriptor = IProjectManager::Get().GetCurrentProject();
@@ -1145,16 +1178,23 @@ bool FPluginManager::ConfigureEnabledPlugins()
 
 					// Copy the plugin references, since we may modify the project if any plugins are missing
 					TArray<FPluginReferenceDescriptor> PluginReferences(ProjectDescriptor->Plugins);
-					for (const FPluginReferenceDescriptor& PluginReference : PluginReferences)
+					if (!ProcessPluginConfigurations(PluginReferences))
 					{
-						if (!ConfiguredPluginNames.Contains(PluginReference.Name))
-						{
-							if (!ConfigureEnabledPluginForCurrentTarget(PluginReference, EnabledPlugins))
-							{
-								return false;
-							}
-							ConfiguredPluginNames.Add(PluginReference.Name);
-						}
+						return false;
+					}
+				}
+			}
+
+			{
+				SCOPED_BOOT_TIMING("PluginConfigServers_PostProjConfig");
+
+				for (const IPluginConfigServer* ConfigServer : PluginConfigServers)
+				{
+					TArray<FPluginReferenceDescriptor> PluginConfigs;
+					ConfigServer->PostProjConfig_GetPluginConfigurations(PluginConfigs);
+					if (!ProcessPluginConfigurations(PluginConfigs))
+					{
+						return false;
 					}
 				}
 			}
diff --git a/Engine/Source/Runtime/Projects/Public/IPluginConfigServer.h b/Engine/Source/Runtime/Projects/Public/IPluginConfigServer.h
new file mode 100644
--- /dev/null
+++ b/Engine/Source/Runtime/Projects/Public/IPluginConfigServer.h
@@ -0,0 +1,38 @@
+// Copyright Epic Games, Inc. All Rights Reserved.
+
+#pragma once
+
+#include "Features/IModularFeature.h"
+#include "PluginReferenceDescriptor.h"
+#include "CoreMinimal.h"
+
+/** Modular feature which allows arbirary systems to extend plugin discovery inside of the engine's plugin manager. */
+class IPluginConfigServer : public IModularFeature
+{
+public:
+	virtual ~IPluginConfigServer() {}
+
+	static FName GetModularFeatureName()
+	{
+		static FName FeatureName = FName(TEXT("PluginConfigServer"));
+		return FeatureName;
+	}
+
+	/**
+	 * Supplies additional configurations for enabling/disabling plugins.
+	 * This is called BEFORE the engine has processed plugin configurations in the .uproject file.
+	 * NOTE: In the case of multiple configurations for the same plugin, the configuration processed first takes priority.
+	 * 
+	 * @param  OutPluginRefs	An array of plugin configurations to append to.
+	 */
+	virtual void PreProjConfig_GetPluginConfigurations(TArray<FPluginReferenceDescriptor>& OutPluginRefs) const = 0;
+
+	/**
+	 * Supplies additional configurations for enabling/disabling plugins.
+	 * This is called AFTER the engine has processed plugin configurations in the .uproject file.
+	 * NOTE: In the case of multiple configurations for the same plugin, the configuration processed first takes priority.
+	 *
+	 * @param  OutPluginRefs	An array of plugin configurations to append to.
+	 */
+	virtual void PostProjConfig_GetPluginConfigurations(TArray<FPluginReferenceDescriptor>& OutPluginRefs) const = 0;
+};
\ No newline at end of file
