Fixing misconfigured ensure (which was blocking functionality).

diff --git a/Engine/Source/Runtime/Projects/Private/PluginManager.cpp b/Engine/Source/Runtime/Projects/Private/PluginManager.cpp
--- a/Engine/Source/Runtime/Projects/Private/PluginManager.cpp
+++ b/Engine/Source/Runtime/Projects/Private/PluginManager.cpp
@@ -219,7 +219,7 @@ namespace DiscoveredPluginMapUtils
 		FDiscoveredPluginMap::ValueType* PluginList = Map.Find(Name);
 
 		// Ensure the specified plugin is in the array
-		if (ensureAlwaysMsgf(PluginList && PluginList->GetData() >= PluginPtr && PluginPtr < PluginList->GetData() + PluginList->Num(),
+		if (ensureAlwaysMsgf(PluginList && PluginList->GetData() <= PluginPtr && PluginPtr < PluginList->GetData() + PluginList->Num(),
 			TEXT("Specified plugin is not registered under the name '%s'. Failed to promote it."), *Name))
 		{
 			TSharedRef<FPlugin>* TopPluginPtr = &PluginList->Top();
