Making some optimizations to plugin discovery in the plugin manager.
Uses the assumption that there shouldn't be other .uplugin files nested under the top level .uplugin directory (no nested .uplugins).

diff --git a/Engine/Source/Runtime/Projects/Private/PluginManager.cpp b/Engine/Source/Runtime/Projects/Private/PluginManager.cpp
--- a/Engine/Source/Runtime/Projects/Private/PluginManager.cpp
+++ b/Engine/Source/Runtime/Projects/Private/PluginManager.cpp
@@ -24,6 +24,8 @@
 #include "Async/ParallelFor.h"
 #include "IPluginConfigServer.h"
 #include "Features/IModularFeatures.h"
+#include "Misc/ScopeRWLock.h"
+#include "Algo/Accumulate.h"
 #if READ_TARGET_ENABLED_PLUGINS_FROM_RECEIPT
 #include "TargetReceipt.h"
 #endif
@@ -770,7 +772,91 @@ void FPluginManager::ReadPluginsInDirectory(const FString& PluginsDirectory, con
 
 void FPluginManager::FindPluginsInDirectory(const FString& PluginsDirectory, TArray<FString>& FileNames)
 {
-	FPlatformFileManager::Get().GetPlatformFile().FindFilesRecursively(FileNames, *PluginsDirectory, TEXT(".uplugin"));
+	//
+	// Use our own custom file discovery method (over something like `IPlatformFile::IterateDirectoryRecursively()`)
+	// to find all the .uplugin files. We utilize optimizations that `IterateDirectoryRecursively()` cannot, because 
+	// we know once we find one .uplugin file that there shouldn't be anymore in the same folder hierarchy.
+	//
+
+	// Directory visitor which aborts iteration once it finds a single .uplugin file
+	class FFindPluginsInDirectory_Visitor : public IPlatformFile::FDirectoryVisitor
+	{
+	public:	
+		FFindPluginsInDirectory_Visitor(TArray<FString>& OutSubDirectories)
+			: IPlatformFile::FDirectoryVisitor(EDirectoryVisitorFlags::ThreadSafe)
+			, SubDirectories(OutSubDirectories)
+		{
+		}
+
+		virtual bool Visit(const TCHAR* FilenameOrDirectory, bool bIsDirectory) override
+		{
+			if (bIsDirectory)
+			{
+				SubDirectories.Emplace(FilenameOrDirectory);
+			}
+			else
+			{
+				constexpr const TCHAR PluginExt[] = TEXT(".uplugin");
+				constexpr const int32 PluginExtLen = UE_ARRAY_COUNT(PluginExt)-1; // -1 for the string's terminating zero (which would be counted by `UE_ARRAY_COUNT()`)
+
+				int32 FileNameLen = FCString::Strlen(FilenameOrDirectory);
+				if (FileNameLen >= PluginExtLen && FCString::Strcmp(&FilenameOrDirectory[FileNameLen - PluginExtLen], PluginExt) == 0)
+				{
+					FoundPluginFile = FilenameOrDirectory;
+
+					// Stop the iteration, we've found the .uplugin we're looking for
+					return false;
+				}
+			}
+			return true;
+		}
+
+		FString FoundPluginFile;
+		TArray<FString>& SubDirectories;
+	};
+
+	TArray<FString> DirectoriesToVisit;
+	DirectoriesToVisit.Add(PluginsDirectory);
+
+	constexpr int32 MinBatchSize = 1;
+	TArray<TArray<FString>> DirectoriesToVisitNext;
+	FRWLock FoundFilesLock;
+	while (DirectoriesToVisit.Num() > 0)
+	{
+		ParallelForWithTaskContext(TEXT("FindPluginsInDirectory.PF"),
+			DirectoriesToVisitNext,
+			DirectoriesToVisit.Num(),
+			MinBatchSize,
+			[&FoundFilesLock, &FileNames, &DirectoriesToVisit](TArray<FString>& OutDirectoriesToVisitNext, int32 Index)
+			{
+				// Track where we start pushing sub-directories to because we might want to discard them (if we end up finding a .uplugin).
+				// Because of how `ParallelForWithTaskContext()` works, this array may already be populated from another execution,
+				// so we have to be targeted about what we clear from the array.
+				const int32 StartingDirIndex = OutDirectoriesToVisitNext.Num();
+
+				FFindPluginsInDirectory_Visitor Visitor(OutDirectoriesToVisitNext); // This visitor writes directly to `OutDirectoriesToVisitNext`, which is why we have to manage its contents
+				FPlatformFileManager::Get().GetPlatformFile().IterateDirectory(*DirectoriesToVisit[Index], Visitor);
+				if (!Visitor.FoundPluginFile.IsEmpty())
+				{
+					// Since we found a .uplugin, ignore sub-directories (stop from iterating deeper) -- 
+					// there shouldn't be any other .uplugin files deeper.
+					// Also, disallow shrinking -- because we're trying to be fast and would rather skip mem reallocs.
+					OutDirectoriesToVisitNext.RemoveAt(StartingDirIndex, OutDirectoriesToVisitNext.Num() - StartingDirIndex, /*bAllowShrinking =*/false);
+
+					// Multiple tasks may be trying to write to this at the same time, lock it
+					FRWScopeLock ScopeLock(FoundFilesLock, SLT_Write);
+					FileNames.Emplace(MoveTemp(Visitor.FoundPluginFile));
+				}
+			},
+			EParallelForFlags::Unbalanced);
+		// Clear and resize `DirectoriesToVisit` for the next batch.
+		DirectoriesToVisit.Reset(Algo::TransformAccumulate(DirectoriesToVisitNext, &TArray<FString>::Num, 0));
+		// Copy all the `DirectoriesToVisitNext` (populated by the various `ParallelFor` tasks) into the one array we use to feed the next round of tasks.
+		for (TArray<FString>& Directories : DirectoriesToVisitNext)
+		{
+			DirectoriesToVisit.Append(MoveTemp(Directories));
+		}
+	}
 }
 
 void FPluginManager::FindPluginManifestsInDirectory(const FString& PluginManifestDirectory, TArray<FString>& FileNames)
@@ -2478,4 +2564,4 @@ void FPluginManager::RemoveFromModuleNameToPluginMap(const TSharedRef<FPlugin>&
 }
 #endif //if WITH_EDITOR
 
-#undef LOCTEXT_NAMESPACE
+#undef LOCTEXT_NAMESPACE
\ No newline at end of file
