Fixing issues with plugins, where calling `RefreshPluginsList()` could discard previously enabled versions of a plugin (when there are multiple version of the same plugin).

diff --git a/Engine/Source/Runtime/Projects/Private/PluginManager.cpp b/Engine/Source/Runtime/Projects/Private/PluginManager.cpp
--- a/Engine/Source/Runtime/Projects/Private/PluginManager.cpp
+++ b/Engine/Source/Runtime/Projects/Private/PluginManager.cpp
@@ -126,6 +126,13 @@ namespace DiscoveredPluginMapUtils
 		return PluginMapValue.Top();
 	}
 
+	/** Removes/Forgets all other known versions of the specified plugin entry. Retains the one "offered" version (which was prioritized and positioned for the rest of UE). */
+	static void DiscardAllSupressedVersions(FDiscoveredPluginMap::ValueType& PluginMapValue)
+	{
+		// The "top" (last) plugin in the list is the one we "offer". Keep it, discard the rest.
+		PluginMapValue.RemoveAtSwap(0, PluginMapValue.Num()-1);
+	}
+
 	/** Returns the one "offered" plugin for the specified plugin name (if any exists). */
 	static TSharedRef<FPlugin>* FindPluginInMap(FDiscoveredPluginMap& InMap, const FDiscoveredPluginMap::KeyType& Name)
 	{
@@ -191,20 +198,6 @@ namespace DiscoveredPluginMapUtils
 			break;
 		}
 	}
-	static void InsertPluginIntoMap_ByHash(FDiscoveredPluginMap& Map, const uint32 KeyHash, const FDiscoveredPluginMap::KeyType& Key, const TSharedRef<FPlugin>& NewPlugin, const EInsertionType InsertionType = EInsertionType::AsOfferedPlugin)
-	{
-		FDiscoveredPluginMap::ValueType& PluginList = Map.FindOrAddByHash(KeyHash, Key);
-		switch (InsertionType)
-		{
-		case EInsertionType::AsOfferedPlugin:
-			InsertPlugin_AsOffered(PluginList, NewPlugin);
-			break;
-
-		case EInsertionType::AsSuppressedPlugin:
-			InsertPlugin_AsSuppressed(PluginList, NewPlugin);
-			break;
-		}
-	}
 
 	/**
 	 * Given how UE is currently structured, we can only load/enable/display one plugin for a given name -- this is the "offered" plugin.
@@ -373,18 +366,14 @@ FPluginManager::~FPluginManager()
 
 void FPluginManager::RefreshPluginsList()
 {
-	// Read a new list of all plugins
-	FDiscoveredPluginMap NewPlugins;
-	ReadAllPlugins(NewPlugins, PluginDiscoveryPaths);
-
-	// Build a list of filenames for plugins which are enabled, and remove the rest
-	TSet<FString> EnabledPluginFileNames;
-	for(FDiscoveredPluginMap::TIterator Iter(AllPlugins); Iter; ++Iter)
+	// Clear out the plugins map (keep only what is enabled)
+	for (FDiscoveredPluginMap::TIterator Iter(AllPlugins); Iter; ++Iter)
 	{
 		const TSharedRef<FPlugin>& Plugin = DiscoveredPluginMapUtils::ResolvePluginFromMapVal(Iter.Value());
-		if(Plugin->bEnabled)
+		if (Plugin->bEnabled)
 		{
-			EnabledPluginFileNames.Add(Plugin->FileName);
+			// Forget all the other discovered versions (which we assume aren't enabled)
+			DiscoveredPluginMapUtils::DiscardAllSupressedVersions(Iter.Value());
 		}
 		else
 		{
@@ -392,25 +381,27 @@ void FPluginManager::RefreshPluginsList()
 		}
 	}
 
-	// Add all the plugins which aren't already enabled
-	for(const FDiscoveredPluginMap::ElementType& NewPluginPair: NewPlugins)
-	{
-		const TSharedRef<FPlugin>& NewPlugin = DiscoveredPluginMapUtils::ResolvePluginFromMapVal(NewPluginPair.Value);
-		if(!EnabledPluginFileNames.Contains(NewPlugin->FileName))
-		{
-			const uint32 PluginNameHash = GetTypeHash(NewPlugin->GetName());
-			DiscoveredPluginMapUtils::InsertPluginIntoMap_ByHash(AllPlugins, PluginNameHash, NewPlugin->GetName(), NewPlugin);
-			PluginsToConfigure.AddByHash(PluginNameHash, NewPlugin->GetName());
-		}
-	}
+	// With the map only containing what was already enabled, fill in the rest 
+	// of the map (duplicate plugins are skipped using a check to the plugin's `FileName` -- see `CreatePluginObject()`)
+	ReadAllPlugins(AllPlugins, PluginDiscoveryPaths);
 
 #if WITH_EDITOR
 	ModuleNameToPluginMap.Reset();
+#endif //if WITH_EDITOR
+
 	for (const FDiscoveredPluginMap::ElementType& PluginPair : AllPlugins)
 	{
-		AddToModuleNameToPluginMap(DiscoveredPluginMapUtils::ResolvePluginFromMapVal(PluginPair.Value));
-	}
+		const TSharedRef<FPlugin>& Plugin = DiscoveredPluginMapUtils::ResolvePluginFromMapVal(PluginPair.Value);
+		if (!Plugin->bEnabled) // if it's not enabled, it's a new plugin
+		{
+			const uint32 PluginNameHash = GetTypeHash(Plugin->GetName());
+			PluginsToConfigure.AddByHash(PluginNameHash, Plugin->GetName());
+		}
+
+#if WITH_EDITOR
+		AddToModuleNameToPluginMap(Plugin);
 #endif //if WITH_EDITOR
+	}
 }
 
 bool VerifySinglePluginForAddOrRemove(const FPluginDescriptor& Descriptor, FText& OutFailReason)
@@ -909,14 +900,19 @@ void FPluginManager::CreatePluginObject(const FString& FileName, const FPluginDe
 		DiscoveredPluginMapUtils::InsertPluginIntoMap(Plugins, Plugin->GetName(), Plugin);
 	}
 	// We allow for duplicates of plugins between engine and the project, favoring the project level plugin
-	else if ((*ExistingPlugin)->Type == EPluginType::Engine && Type == EPluginType::Project)
+	else if ((*ExistingPlugin)->Type == EPluginType::Engine && Type == EPluginType::Project && !(*ExistingPlugin)->bEnabled)
 	{
 		UE_LOG(LogPluginManager, Verbose, TEXT("By default, prioritizing project plugin (%s) over the corresponding engine version (%s)."), *Plugin->FileName, *(*ExistingPlugin)->FileName);
 		DiscoveredPluginMapUtils::InsertPluginIntoMap(Plugins, Plugin->GetName(), Plugin, DiscoveredPluginMapUtils::EInsertionType::AsOfferedPlugin);
 	}
 	else if ((*ExistingPlugin)->FileName != Plugin->FileName)
 	{
-		if ((*ExistingPlugin)->Type == EPluginType::Project && Type == EPluginType::Engine)
+		if ((*ExistingPlugin)->bEnabled)
+		{
+			UE_LOG(LogPluginManager, Verbose, TEXT("A version of the '%s' plugin has already been enabled (%s); prioritizing that over the newly discovered version (%s)."), *Plugin->Name, *(*ExistingPlugin)->FileName, *Plugin->FileName);
+			DiscoveredPluginMapUtils::InsertPluginIntoMap(Plugins, Plugin->GetName(), Plugin, DiscoveredPluginMapUtils::EInsertionType::AsSuppressedPlugin);
+		}
+		else if ((*ExistingPlugin)->Type == EPluginType::Project && Type == EPluginType::Engine)
 		{
 			// Project plugins are favored over engine plugins, so we don't want to warn in this case
 			// (instead we mimic the Verbose log from above, and just explain which plugin we're favoring)
@@ -1764,30 +1760,40 @@ bool FPluginManager::ConfigureEnabledPluginForTarget(const FPluginReferenceDescr
 			}
 
 			// Find the plugin being enabled
-			const TSharedRef<FPlugin>* PluginPtr = nullptr;
-			if (Reference.RequestedVersion.IsSet())
+			const TSharedRef<FPlugin>* PluginPtr = DiscoveredPluginMapUtils::FindPluginInMap(AllPlugins, Reference.Name);
+			if (PluginPtr && Reference.RequestedVersion.IsSet())
 			{
-				if (TArray<TSharedRef<FPlugin>>* PluginVersions = DiscoveredPluginMapUtils::FindAllPluginVersionsWithName(AllPlugins, Reference.Name))
+				TArray<TSharedRef<FPlugin>>* PluginVersions = DiscoveredPluginMapUtils::FindAllPluginVersionsWithName(AllPlugins, Reference.Name);
+				check(PluginVersions);
+
+				TSharedRef<FPlugin>* FoundVersion = nullptr;
+				for (TSharedRef<FPlugin>& PluginVersion : *PluginVersions)
 				{
-					for (TSharedRef<FPlugin>& PluginVersion : *PluginVersions)
+					if (PluginVersion->GetDescriptor().Version == Reference.RequestedVersion.GetValue())
 					{
-						if (PluginVersion->GetDescriptor().Version == Reference.RequestedVersion.GetValue())
-						{
-							PluginPtr = DiscoveredPluginMapUtils::PromotePluginToOfferedVersion(AllPlugins, &PluginVersion);
-							break;
-						}
+						FoundVersion = &PluginVersion;
+						break;
 					}
+				}
 
-					if (!PluginPtr)
-					{
-						UE_LOG(LogPluginManager, Warning, TEXT("Failed to find specific version (v%d) of plugin '%s'. Other versions exist, but the explicit version requested was missing."),
-							FirstReference.RequestedVersion.GetValue(), *Reference.Name);
-					}
+				if (!FoundVersion)
+				{
+					UE_LOG(LogPluginManager, Warning, TEXT("Failed to find specific version (v%d) of plugin '%s'. Other versions exist, but the explicit version requested was missing."),
+						Reference.RequestedVersion.GetValue(), *Reference.Name);
+
+					PluginPtr = nullptr;
+				}
+				else if (FoundVersion != PluginPtr && (*PluginPtr)->bEnabled)
+				{
+					UE_LOG(LogPluginManager, Warning, TEXT("A different version of the %s plugin (v%d : '%s') is already enabled. Cannot switch to v%d (%s) while another version is active."),
+						*Reference.Name, (*PluginPtr)->Descriptor.Version, *(*PluginPtr)->FileName, Reference.RequestedVersion.GetValue(), *(*FoundVersion)->FileName);
+
+					PluginPtr = nullptr;
+				}
+				else
+				{
+					PluginPtr = DiscoveredPluginMapUtils::PromotePluginToOfferedVersion(AllPlugins, FoundVersion);
 				}
-			}
-			else
-			{
-				PluginPtr = DiscoveredPluginMapUtils::FindPluginInMap(AllPlugins, Reference.Name);
 			}
 
 			if (PluginPtr == nullptr)
