[Engine Runtime]
PluginManager - Disabled ini file find caching for all windows configs (not only editor).. this should probably be disabled for everything but since I can't test it easily I will leave that to someone else to change.

This saves ~3 seconds of loading in to some Epic games

diff --git a/Engine/Source/Runtime/Projects/Private/PluginManager.cpp b/Engine/Source/Runtime/Projects/Private/PluginManager.cpp
--- a/Engine/Source/Runtime/Projects/Private/PluginManager.cpp
+++ b/Engine/Source/Runtime/Projects/Private/PluginManager.cpp
@@ -86,7 +86,7 @@ namespace PluginSystemDefs
 
 	bool IsCachingIniFilesForProcessing()
 	{
-#if WITH_EDITOR
+#if PLATFORM_WINDOWS // return false should be false for all platforms according to Josh Adams but has not been tested on other platforms than windows
 		return false;
 #endif
 		return true;
