Fix problem with Platform Extensions that are not calling StartupModule on added modules. Problem reported on UDN.

diff --git a/Engine/Source/Runtime/Projects/Private/PluginManager.cpp b/Engine/Source/Runtime/Projects/Private/PluginManager.cpp
--- a/Engine/Source/Runtime/Projects/Private/PluginManager.cpp
+++ b/Engine/Source/Runtime/Projects/Private/PluginManager.cpp
@@ -678,11 +678,14 @@ void FPluginManager::ReadAllPlugins(FDiscoveredPluginMap& Plugins, const TSet<FS
 			TSharedRef<FPlugin>& Parent = *ParentPtr;
 			for (const FModuleDescriptor& ChildModule : Child->GetDescriptor().Modules)
 			{
+				bool bFound = false;
+
 				// look for a matching parent
 				for (FModuleDescriptor& ParentModule : Parent->Descriptor.Modules)
 				{
 					if (ParentModule.Name == ChildModule.Name && ParentModule.Type == ChildModule.Type)
 					{
+						bFound = true;
 						// we only need to add the platform to an allow list if the parent had an allow list (otherwise, we could mistakenly remove all other platforms)
 						if (ParentModule.bHasExplicitPlatforms || ParentModule.PlatformAllowList.Num() > 0)
 						{
@@ -693,6 +696,12 @@ void FPluginManager::ReadAllPlugins(FDiscoveredPluginMap& Plugins, const TSet<FS
 						ParentModule.PlatformDenyList.Append(ChildModule.PlatformDenyList);
 					}
 				}
+				
+				if (!bFound)
+				{
+					Parent->Descriptor.Modules.Add(ChildModule);
+				}
+
 			}
 
 			if (Parent->GetDescriptor().bHasExplicitPlatforms || Parent->GetDescriptor().SupportedTargetPlatforms.Num() != 0)
