IPlugin::GetVersePath

diff --git a/Engine/Source/Runtime/Projects/Private/PluginManager.h b/Engine/Source/Runtime/Projects/Private/PluginManager.h
--- a/Engine/Source/Runtime/Projects/Private/PluginManager.h
+++ b/Engine/Source/Runtime/Projects/Private/PluginManager.h
@@ -82,6 +82,11 @@ public:
 		return Descriptor.bCanContainVerse;
 	}
 
+	virtual const FString& GetVersePath() const override
+	{
+		return Descriptor.VersePath;
+	}
+
 	virtual EPluginType GetType() const override
 	{
 		return Type;
diff --git a/Engine/Source/Runtime/Projects/Public/Interfaces/IPluginManager.h b/Engine/Source/Runtime/Projects/Public/Interfaces/IPluginManager.h
--- a/Engine/Source/Runtime/Projects/Public/Interfaces/IPluginManager.h
+++ b/Engine/Source/Runtime/Projects/Public/Interfaces/IPluginManager.h
@@ -86,23 +86,23 @@ public:
 	virtual const FString& GetFriendlyName() const = 0;
 
 	/**
-	 * Get a path to the plugin's descriptor
+	 * Get a filesystem path to the plugin's descriptor
 	 *
-	 * @return Path to the plugin's descriptor.
+	 * @return Filesystem path to the plugin's descriptor.
 	 */
 	virtual const FString& GetDescriptorFileName() const = 0;
 
 	/**
-	 * Get a path to the plugin's directory.
+	 * Get a filesystem path to the plugin's directory.
 	 *
-	 * @return Path to the plugin's base directory.
+	 * @return Filesystem path to the plugin's base directory.
 	 */
 	virtual FString GetBaseDir() const = 0;
 
 	/**
-	 * Get a path to the plugin's content directory.
+	 * Get a filesystem path to the plugin's content directory.
 	 *
-	 * @return Path to the plugin's content directory.
+	 * @return Filesystem path to the plugin's content directory.
 	 */
 	virtual FString GetContentDir() const = 0;
 
@@ -155,6 +155,13 @@ public:
 	 */
 	virtual bool CanContainVerse() const = 0;
 
+	/**
+	 * Gets the Verse path to the root of the plugin's content directory
+	 *
+	 * @return Verse path to the root of the plugin's content directory
+	 */
+	virtual const FString& GetVersePath() const = 0;
+
 	/**
 	 * Returns the plugin's location
 	 *
