Undoing some questionable plugin manager changes, which encouraged "chicken-and-egg" plugin scenarios (a plugin which could introduce more auto-loaded plugins).

Instead, expanding the explicit mounting behavior, providing alternative signatures to MountExplicitlyLoadedPlugin(), encouraging users to use the already established way of dynamically mounting plugins.

diff --git a/Engine/Source/Runtime/Projects/Private/PluginManager.cpp b/Engine/Source/Runtime/Projects/Private/PluginManager.cpp
--- a/Engine/Source/Runtime/Projects/Private/PluginManager.cpp
+++ b/Engine/Source/Runtime/Projects/Private/PluginManager.cpp
@@ -23,8 +23,6 @@
 #include "PluginManifest.h"
 #include "HAL/PlatformTime.h"
 #include "Async/ParallelFor.h"
-#include "IPluginConfigServer.h"
-#include "Features/IModularFeatures.h"
 #include "Misc/ScopeRWLock.h"
 #include "Algo/Accumulate.h"
 #if READ_TARGET_ENABLED_PLUGINS_FROM_RECEIPT
@@ -133,6 +131,12 @@ namespace DiscoveredPluginMapUtils
 		PluginMapValue.RemoveAtSwap(0, PluginMapValue.Num()-1);
 	}
 
+	/** Returns an array of all the discovered plugins with a matching name (null if none were discovered). */
+	static TArray<TSharedRef<FPlugin>>* FindAllPluginVersionsWithName(FDiscoveredPluginMap& InMap, const FDiscoveredPluginMap::KeyType& Name)
+	{
+		return InMap.Find(Name);
+	}
+
 	/** Returns the one "offered" plugin for the specified plugin name (if any exists). */
 	static TSharedRef<FPlugin>* FindPluginInMap(FDiscoveredPluginMap& InMap, const FDiscoveredPluginMap::KeyType& Name)
 	{
@@ -164,10 +168,60 @@ namespace DiscoveredPluginMapUtils
 		return *PluginPtr;
 	}
 
-	/** Returns an array of all the discovered plugins with a matching name (null if none were discovered). */
-	static TArray<TSharedRef<FPlugin>>* FindAllPluginVersionsWithName(FDiscoveredPluginMap& InMap, const FDiscoveredPluginMap::KeyType& Name)
+	static TSharedRef<FPlugin>* FindPluginInMap_FromFileName(FDiscoveredPluginMap& InMap, const FString& FileName)
 	{
-		return InMap.Find(Name);
+		const FString PluginName = FPaths::GetBaseFilename(FileName);
+		TArray<TSharedRef<FPlugin>>* KnownVersions = DiscoveredPluginMapUtils::FindAllPluginVersionsWithName(InMap, PluginName);
+		if (KnownVersions)
+		{
+			FString NormalizedFilenameToFind = FileName;
+			FPaths::NormalizeFilename(NormalizedFilenameToFind);
+
+			for (TSharedRef<FPlugin>& Plugin : *KnownVersions)
+			{
+				FString NormalizedPluginPath = Plugin->FileName;
+				FPaths::NormalizeFilename(NormalizedPluginPath);
+				if (NormalizedPluginPath.Equals(NormalizedFilenameToFind))
+				{
+					return &Plugin;
+				}
+			}
+		}
+		return nullptr;
+	}
+
+	static TSharedRef<FPlugin>* FindPluginInMap_FromDescriptor(FDiscoveredPluginMap& InMap, const FPluginReferenceDescriptor& PluginDesc)
+	{
+		TArray<TSharedRef<FPlugin>>* KnownVersions = DiscoveredPluginMapUtils::FindAllPluginVersionsWithName(InMap, PluginDesc.Name);
+		if (KnownVersions && KnownVersions->Num() > 0)
+		{
+			if (PluginDesc.RequestedVersion.IsSet())
+			{
+				for (TSharedRef<FPlugin>& Plugin : *KnownVersions)
+				{
+					if (Plugin->GetDescriptor().Version == PluginDesc.RequestedVersion)
+					{
+						return &Plugin;
+					}
+				}
+			}
+			else
+			{
+				return &KnownVersions->Top();
+			}
+		}
+		return nullptr;
+	}
+
+	static bool IsOfferedPlugin(const FDiscoveredPluginMap& InMap, const TSharedRef<FPlugin>& Plugin)
+	{
+		const FDiscoveredPluginMap::ValueType* PluginList = InMap.Find(Plugin->Name);
+		if (PluginList && !PluginList->IsEmpty())
+		{
+			// The "top" plugin in the list is the one we choose to present to the rest of UE
+			return Plugin == PluginList->Top();
+		}
+		return false;
 	}
 
 	/** Internal insersion methods used for controlling the one "offered" plugin we share with the rest of the engine. */
@@ -199,6 +253,15 @@ namespace DiscoveredPluginMapUtils
 		}
 	}
 
+	static void RemovePluginFromMap(FDiscoveredPluginMap& InMap, const TSharedRef<FPlugin>& Plugin)
+	{
+		FDiscoveredPluginMap::ValueType* PluginList = InMap.Find(Plugin->Name);
+		if (PluginList && !PluginList->IsEmpty())
+		{
+			PluginList->Remove(Plugin);
+		}
+	}
+
 	/**
 	 * Given how UE is currently structured, we can only load/enable/display one plugin for a given name -- this is the "offered" plugin.
 	 * This method let's us swap out the current "offered" plugin with the one we specify.
@@ -208,7 +271,7 @@ namespace DiscoveredPluginMapUtils
 	 * 
 	 * @return Null if the promotion failed, otherwise returns a new TSharedRef<> referencing the plugin that `PluginPtr` was initially.
 	 */
-	static const TSharedRef<FPlugin>* PromotePluginToOfferedVersion(FDiscoveredPluginMap& Map, TSharedRef<FPlugin>* PluginPtr)
+	static TSharedRef<FPlugin>* PromotePluginToOfferedVersion(FDiscoveredPluginMap& Map, TSharedRef<FPlugin>* PluginPtr)
 	{
 		const FString& Name = (*PluginPtr)->GetName();
 		FDiscoveredPluginMap::ValueType* PluginList = Map.Find(Name);
@@ -420,9 +483,20 @@ bool FPluginManager::AddToPluginsList(const FString& PluginFilename, FText* OutF
 #if (WITH_ENGINE && !IS_PROGRAM) || WITH_PLUGIN_SUPPORT
 	// No need to re-add if it already exists
 	FString PluginName = FPaths::GetBaseFilename(PluginFilename);
-	if (AllPlugins.Contains(PluginName))
+	const TArray<TSharedRef<FPlugin>>* KnownVersions = DiscoveredPluginMapUtils::FindAllPluginVersionsWithName(AllPlugins, PluginName);
+	if (KnownVersions)
 	{
-		return true;
+		FString NormalizedFilename = PluginFilename;
+		FPaths::NormalizeFilename(NormalizedFilename);
+		for (const TSharedRef<FPlugin>& Plugin : *KnownVersions)
+		{
+			FString NormalizedPluginPath = Plugin->FileName;
+			FPaths::NormalizeFilename(NormalizedPluginPath);
+			if (NormalizedPluginPath.Equals(NormalizedFilename))
+			{
+				return true;
+			}
+		}
 	}
 
 	// Read the plugin and load it
@@ -466,10 +540,17 @@ bool FPluginManager::AddToPluginsList(const FString& PluginFilename, FText* OutF
 		TSharedRef<FPlugin>* NewPlugin = DiscoveredPluginMapUtils::FindPluginInMap(NewPlugins, PluginName);
 		if (ensure(NewPlugin))
 		{
-			DiscoveredPluginMapUtils::InsertPluginIntoMap(AllPlugins, PluginName, *NewPlugin);
+			// Maintain behavior precedence for this function -- the new plugin doesn't superscede any previous versions (but we still add it for tracking)
+			const bool bPluginAlreadyExists = KnownVersions && KnownVersions->Num() > 0;
+			const DiscoveredPluginMapUtils::EInsertionType Priority = (bPluginAlreadyExists) ? DiscoveredPluginMapUtils::EInsertionType::AsSuppressedPlugin : DiscoveredPluginMapUtils::EInsertionType::AsOfferedPlugin;
+
+			DiscoveredPluginMapUtils::InsertPluginIntoMap(AllPlugins, PluginName, *NewPlugin, Priority);
 
 #if WITH_EDITOR
-			AddToModuleNameToPluginMap(*NewPlugin);
+			if (Priority == DiscoveredPluginMapUtils::EInsertionType::AsOfferedPlugin)
+			{
+				AddToModuleNameToPluginMap(*NewPlugin);
+			}
 #endif //if WITH_EDITOR
 		}
 
@@ -492,9 +573,7 @@ bool FPluginManager::AddToPluginsList(const FString& PluginFilename, FText* OutF
 bool FPluginManager::RemoveFromPluginsList(const FString& PluginFilename, FText* OutFailReason /*= nullptr*/)
 {
 #if (WITH_ENGINE && !IS_PROGRAM) || WITH_PLUGIN_SUPPORT
-	FString PluginName = FPaths::GetBaseFilename(PluginFilename);
-
-	TSharedRef<FPlugin>* MaybePlugin = DiscoveredPluginMapUtils::FindPluginInMap(AllPlugins, PluginName);
+	TSharedRef<FPlugin>* MaybePlugin = DiscoveredPluginMapUtils::FindPluginInMap_FromFileName(AllPlugins, PluginFilename);
 	if (MaybePlugin == nullptr)
 	{
 		return true;
@@ -525,10 +604,14 @@ bool FPluginManager::RemoveFromPluginsList(const FString& PluginFilename, FText*
 	}
 
 #if WITH_EDITOR
-	RemoveFromModuleNameToPluginMap(FoundPlugin);
+	// Is this the plugin that would have been mapped?
+	if (DiscoveredPluginMapUtils::IsOfferedPlugin(AllPlugins, *MaybePlugin))
+	{
+		RemoveFromModuleNameToPluginMap(FoundPlugin);
+	}
 #endif //if WITH_EDITOR
 
-	AllPlugins.Remove(PluginName);
+	DiscoveredPluginMapUtils::RemovePluginFromMap(AllPlugins, FoundPlugin);
 #endif //if (WITH_ENGINE && !IS_PROGRAM) || WITH_PLUGIN_SUPPORT
 
 	return true;
@@ -1243,21 +1326,6 @@ bool FPluginManager::ConfigureEnabledPlugins()
 				return true;
 			};
 
-			TArray<IPluginConfigServer*> PluginConfigServers = IModularFeatures::Get().GetModularFeatureImplementations<IPluginConfigServer>(IPluginConfigServer::GetModularFeatureName());
-			{
-				SCOPED_BOOT_TIMING("PluginConfigServers_PreProjConfig");
-				
-				for (const IPluginConfigServer* ConfigServer : PluginConfigServers)
-				{
-					TArray<FPluginReferenceDescriptor> PluginConfigs;
-					ConfigServer->PreProjConfig_GetPluginConfigurations(PluginConfigs);
-					if (!ProcessPluginConfigurations(PluginConfigs))
-					{
-						return false;
-					}					
-				}
-			}
-
 			bool bAllowEnginePluginsEnabledByDefault = true;
 			// Find all the plugin references in the project file
 			const FProjectDescriptor* ProjectDescriptor = IProjectManager::Get().GetCurrentProject();
@@ -1277,20 +1345,6 @@ bool FPluginManager::ConfigureEnabledPlugins()
 				}
 			}
 
-			{
-				SCOPED_BOOT_TIMING("PluginConfigServers_PostProjConfig");
-
-				for (const IPluginConfigServer* ConfigServer : PluginConfigServers)
-				{
-					TArray<FPluginReferenceDescriptor> PluginConfigs;
-					ConfigServer->PostProjConfig_GetPluginConfigurations(PluginConfigs);
-					if (!ProcessPluginConfigurations(PluginConfigs))
-					{
-						return false;
-					}
-				}
-			}
-
 			{
 				// Add the plugins which are enabled by default
 				SCOPED_BOOT_TIMING("AddPluginsEnabledByDefault");
@@ -2382,13 +2436,63 @@ void FPluginManager::MountNewlyCreatedPlugin(const FString& PluginName)
 	}
 }
 
-void FPluginManager::MountExplicitlyLoadedPlugin(const FString& PluginName)
+bool FPluginManager::MountExplicitlyLoadedPlugin(const FString& PluginName)
 {
+	bool bSuccess = false;
+
 	TSharedPtr<FPlugin> Plugin = FindPluginInstance(PluginName);
 	if (Plugin.IsValid() && Plugin->Descriptor.bExplicitlyLoaded)
 	{
 		MountPluginFromExternalSource(Plugin.ToSharedRef());
+		bSuccess = true;
+	}
+
+	return bSuccess;
+}
+
+bool FPluginManager::MountExplicitlyLoadedPlugin_FromFileName(const FString& PluginFileName)
+{
+	TSharedRef<FPlugin>* DescribedPlugin = DiscoveredPluginMapUtils::FindPluginInMap_FromFileName(AllPlugins, PluginFileName);
+	return TryMountExplicitlyLoadedPluginVersion(DescribedPlugin);
+}
+
+bool FPluginManager::MountExplicitlyLoadedPlugin_FromDescriptor(const FPluginReferenceDescriptor& PluginDescriptor)
+{
+	TSharedRef<FPlugin>* DescribedPlugin = DiscoveredPluginMapUtils::FindPluginInMap_FromDescriptor(AllPlugins, PluginDescriptor);
+	return TryMountExplicitlyLoadedPluginVersion(DescribedPlugin);
+}
+
+bool FPluginManager::TryMountExplicitlyLoadedPluginVersion(TSharedRef<FPlugin>* AllPlugins_PluginPtr)
+{
+	bool bSuccess = false;
+
+	if (AllPlugins_PluginPtr && (*AllPlugins_PluginPtr)->Descriptor.bExplicitlyLoaded)
+	{
+		TSharedPtr<FPlugin> CurrentPriorityPlugin = FindPluginInstance((*AllPlugins_PluginPtr)->GetName());
+		if (ensure(CurrentPriorityPlugin.IsValid()) && CurrentPriorityPlugin != *AllPlugins_PluginPtr && CurrentPriorityPlugin->IsEnabled())
+		{
+			UE_LOG(LogPluginManager, Error, TEXT("Cannot mount plugin '%s' with another version (%s) enabled."), *(*AllPlugins_PluginPtr)->FileName, *CurrentPriorityPlugin->FileName);
+		}
+		else
+		{
+#if WITH_EDITOR
+			RemoveFromModuleNameToPluginMap(CurrentPriorityPlugin.ToSharedRef());
+			AddToModuleNameToPluginMap(*AllPlugins_PluginPtr);
+#endif //if WITH_EDITOR
+
+			// NOTE: This mutates what the passed `DescribedPlugin` pointer is referencing, which is why 
+			//       we re-assign it to the return value (to keep it referencing the same plugin)
+			AllPlugins_PluginPtr = DiscoveredPluginMapUtils::PromotePluginToOfferedVersion(AllPlugins, AllPlugins_PluginPtr);
+			if (ensure(AllPlugins_PluginPtr))
+			{
+				MountPluginFromExternalSource(*AllPlugins_PluginPtr);
+
+				bSuccess = true;
+			}
+		}
 	}
+
+	return bSuccess;
 }
 
 void FPluginManager::MountPluginFromExternalSource(const TSharedRef<FPlugin>& Plugin)
diff --git a/Engine/Source/Runtime/Projects/Private/PluginManager.h b/Engine/Source/Runtime/Projects/Private/PluginManager.h
--- a/Engine/Source/Runtime/Projects/Private/PluginManager.h
+++ b/Engine/Source/Runtime/Projects/Private/PluginManager.h
@@ -156,7 +156,9 @@ public:
 	virtual FNewPluginMountedEvent& OnNewPluginMounted() override;
 	virtual FNewPluginMountedEvent& OnPluginEdited() override;
 	virtual void MountNewlyCreatedPlugin(const FString& PluginName) override;
-	virtual void MountExplicitlyLoadedPlugin(const FString& PluginName) override;
+	virtual bool MountExplicitlyLoadedPlugin(const FString& PluginName) override;
+	virtual bool MountExplicitlyLoadedPlugin_FromFileName(const FString& PluginFileName) override;
+	virtual bool MountExplicitlyLoadedPlugin_FromDescriptor(const FPluginReferenceDescriptor& PluginDescriptor) override;
 	virtual bool UnmountExplicitlyLoadedPlugin(const FString& PluginName, FText* OutReason) override;
 	virtual FName PackageNameFromModuleName(FName ModuleName) override;
 	virtual bool RequiresTempTargetForCodePlugin(const FProjectDescriptor* ProjectDescriptor, const FString& Platform, EBuildConfiguration Configuration, EBuildTargetType TargetType, FText& OutReason) override;
@@ -221,6 +223,15 @@ private:
 	/** Gets the instance of a given plugin */
 	TSharedPtr<FPlugin> FindPluginInstance(const FString& Name);
 
+	/** 
+	 * Attempts to mount a spectific plugin version. Can fail if the plugin isn't marked 'ExplicitlyLoaded',
+	 * or if there's a different version already mounted.
+	 * 
+	 * NOTE: It's expected that `AllPlugins_PluginPtr` directly addresses an entry in `AllPlugins` 
+	 *       (so that it can reorder the versions in `AllPlugins` since that's how we track the mounted/choice version).
+	 */
+	bool TryMountExplicitlyLoadedPluginVersion(TSharedRef<FPlugin>* AllPlugins_PluginPtr);
+
 	/** Mounts a plugin that was requested to be mounted from external code (either by MountNewlyCreatedPlugin or MountExplicitlyLoadedPlugin) */
 	void MountPluginFromExternalSource(const TSharedRef<FPlugin>& Plugin);
 
diff --git a/Engine/Source/Runtime/Projects/Public/IPluginConfigServer.h b/Engine/Source/Runtime/Projects/Public/IPluginConfigServer.h
deleted file mode 100644
diff --git a/Engine/Source/Runtime/Projects/Public/Interfaces/IPluginManager.h b/Engine/Source/Runtime/Projects/Public/Interfaces/IPluginManager.h
--- a/Engine/Source/Runtime/Projects/Public/Interfaces/IPluginManager.h
+++ b/Engine/Source/Runtime/Projects/Public/Interfaces/IPluginManager.h
@@ -7,6 +7,7 @@
 
 struct FProjectDescriptor;
 class FJsonObject;
+struct FPluginReferenceDescriptor;
 
 /**
  * Enum for where a plugin is loaded from
@@ -399,7 +400,9 @@ public:
 	 * Marks an explicitly loaded plugin as enabled, mounts its content and tries to load its modules.
 	 * These plugins are not loaded implicitly, but instead wait for this function to be called.
 	 */
-	virtual void MountExplicitlyLoadedPlugin(const FString& PluginName) = 0;
+	virtual bool MountExplicitlyLoadedPlugin(const FString& PluginName) = 0;
+	virtual bool MountExplicitlyLoadedPlugin_FromFileName(const FString& PluginFileName) = 0;
+	virtual bool MountExplicitlyLoadedPlugin_FromDescriptor(const FPluginReferenceDescriptor& PluginDescriptor) = 0;
 
 	/**
 	 * Marks an explicitly loaded plugin as disabled, unmounts its content (does not work on plugins with compiled modules).
