Limit loading phase ensures to !UE_EDITOR configurations.
We now know what mode we are running in before AppInit() in other configurations.

diff --git a/Engine/Source/Runtime/Projects/Private/ModuleDescriptor.cpp b/Engine/Source/Runtime/Projects/Private/ModuleDescriptor.cpp
--- a/Engine/Source/Runtime/Projects/Private/ModuleDescriptor.cpp
+++ b/Engine/Source/Runtime/Projects/Private/ModuleDescriptor.cpp
@@ -622,24 +622,30 @@ bool FModuleDescriptor::IsLoadedInCurrentConfiguration() const
 
 	case EHostType::Editor:
 		#if WITH_EDITOR
-			// GIsEditor is not set until the PostSplashScreen phase
-			ensure(LoadingPhase != ELoadingPhase::PostConfigInit && LoadingPhase != ELoadingPhase::EarliestPossible);
+			#if !UE_EDITOR
+				// SetIsRunningAsCommandlet() may be set after AppInit() via "late commandlet token" path
+				ensure(LoadingPhase != ELoadingPhase::PostConfigInit && LoadingPhase != ELoadingPhase::EarliestPossible);
+			#endif
 			if(GIsEditor) return true;
 		#endif
 		break;
 
 	case EHostType::EditorNoCommandlet:
 		#if WITH_EDITOR
-			// GIsEditor is not set until the PostSplashScreen phase
-			ensure(LoadingPhase != ELoadingPhase::PostConfigInit && LoadingPhase != ELoadingPhase::EarliestPossible);
+			#if !UE_EDITOR
+				// SetIsRunningAsCommandlet() may be set after AppInit() via "late commandlet token" path
+				ensure(LoadingPhase != ELoadingPhase::PostConfigInit && LoadingPhase != ELoadingPhase::EarliestPossible);
+			#endif
 			if(GIsEditor && !IsRunningCommandlet()) return true;
 		#endif
 		break;
 
 	case EHostType::EditorAndProgram:
 		#if WITH_EDITOR
-			// GIsEditor is not set until the PostSplashScreen phase
-			ensure(LoadingPhase != ELoadingPhase::PostConfigInit && LoadingPhase != ELoadingPhase::EarliestPossible);
+			#if !UE_EDITOR
+				// SetIsRunningAsCommandlet() may be set after AppInit() via "late commandlet token" path
+				ensure(LoadingPhase != ELoadingPhase::PostConfigInit && LoadingPhase != ELoadingPhase::EarliestPossible);
+			#endif
 			return GIsEditor;
 		#elif IS_PROGRAM
 			return true;
