Adding https:// conditions to links in the Plugin Browser

diff --git a/Engine/Plugins/Editor/PluginBrowser/Source/PluginBrowser/Private/SPluginTile.cpp b/Engine/Plugins/Editor/PluginBrowser/Source/PluginBrowser/Private/SPluginTile.cpp
--- a/Engine/Plugins/Editor/PluginBrowser/Source/PluginBrowser/Private/SPluginTile.cpp
+++ b/Engine/Plugins/Editor/PluginBrowser/Source/PluginBrowser/Private/SPluginTile.cpp
@@ -83,7 +83,7 @@ void SPluginTile::RecreateWidgets()
 	// create support link
 	TSharedPtr<SWidget> SupportWidget;
 	{
-		if (PluginDescriptor.SupportURL.IsEmpty())
+		if (PluginDescriptor.SupportURL.IsEmpty() || !PluginDescriptor.SupportURL.StartsWith("https://"))
 		{
 			SupportWidget = SNullWidget::NullWidget;
 		}
@@ -118,7 +118,7 @@ void SPluginTile::RecreateWidgets()
 	// create documentation link
 	TSharedPtr<SWidget> DocumentationWidget;
 	{
-		if (PluginDescriptor.DocsURL.IsEmpty())
+		if (PluginDescriptor.DocsURL.IsEmpty() || !PluginDescriptor.DocsURL.StartsWith("https://"))
 		{
 			DocumentationWidget = SNullWidget::NullWidget;
 		}
@@ -157,7 +157,7 @@ void SPluginTile::RecreateWidgets()
 		{
 			CreatedByWidget = SNullWidget::NullWidget;
 		}
-		else if (PluginDescriptor.CreatedByURL.IsEmpty())
+		else if (PluginDescriptor.CreatedByURL.IsEmpty() || !PluginDescriptor.CreatedByURL.StartsWith("https://"))
 		{
 			CreatedByWidget = SNew(SHorizontalBox)
 
diff --git a/Engine/Source/Runtime/Projects/Private/PluginManager.cpp b/Engine/Source/Runtime/Projects/Private/PluginManager.cpp
--- a/Engine/Source/Runtime/Projects/Private/PluginManager.cpp
+++ b/Engine/Source/Runtime/Projects/Private/PluginManager.cpp
@@ -1958,17 +1958,20 @@ bool FPluginManager::ConfigureEnabledPluginForTarget(const FPluginReferenceDescr
 
 bool FPluginManager::PromptToDownloadPlugin(const FString& PluginName, const FString& MarketplaceURL)
 {
-	FText Caption = FText::Format(LOCTEXT("DownloadPluginCaption", "Missing {0} Plugin"), FText::FromString(PluginName));
-	FText Message = FText::Format(LOCTEXT("DownloadPluginMessage", "This project requires the {0} plugin.\n\nWould you like to download it from the Unreal Engine Marketplace?"), FText::FromString(PluginName));
-	if(FMessageDialog::Open(EAppMsgType::YesNo, Message, &Caption) == EAppReturnType::Yes)
+	if(MarketplaceURL.StartsWith("https://"))
 	{
-		FString Error;
-		FPlatformProcess::LaunchURL(*MarketplaceURL, nullptr, &Error);
-		if (Error.Len() == 0)
+		FText Caption = FText::Format(LOCTEXT("DownloadPluginCaption", "Missing {0} Plugin"), FText::FromString(PluginName));
+		FText Message = FText::Format(LOCTEXT("DownloadPluginMessage", "This project requires the {0} plugin.\n\nWould you like to download it from the Unreal Engine Marketplace?"), FText::FromString(PluginName));
+		if(FMessageDialog::Open(EAppMsgType::YesNo, Message, &Caption) == EAppReturnType::Yes)
 		{
-			return true;
+			FString Error;
+			FPlatformProcess::LaunchURL(*MarketplaceURL, nullptr, &Error);
+			if (Error.Len() == 0)
+			{
+				return true;
+			}
+			FMessageDialog::Open(EAppMsgType::Ok, FText::FromString(Error));
 		}
-		FMessageDialog::Open(EAppMsgType::Ok, FText::FromString(Error));
 	}
 	return false;
 }
