Fix for merge error due to changed AllPlugins container type.

diff --git a/Engine/Source/Runtime/Projects/Private/PluginManager.cpp b/Engine/Source/Runtime/Projects/Private/PluginManager.cpp
--- a/Engine/Source/Runtime/Projects/Private/PluginManager.cpp
+++ b/Engine/Source/Runtime/Projects/Private/PluginManager.cpp
@@ -2645,14 +2645,16 @@ bool FPluginManager::TrySplitVersePath(const UE::Core::FVersePath& VersePath, FN
 
 	FStringView VersePathView = VersePath.AsStringView();
 
-	for (const TPair<FString, TSharedRef<FPlugin>>& NamePluginPair : AllPlugins)
+	for (const TPair<FString, TArray<TSharedRef<FPlugin>>>& NamePluginPair : AllPlugins)
 	{
-		const FPluginDescriptor& PluginDescriptor = NamePluginPair.Value->Descriptor;
+		const TSharedRef<FPlugin>& Plugin = DiscoveredPluginMapUtils::ResolvePluginFromMapVal(NamePluginPair.Value);
+
+		const FPluginDescriptor& PluginDescriptor = Plugin->Descriptor;
 		if (!PluginDescriptor.VersePath.IsEmpty() && VersePathView.StartsWith(PluginDescriptor.VersePath + TEXT('/')))
 		{
 			VersePathView.RightChopInline(PluginDescriptor.VersePath.Len());
 
-			OutPackageName = *NamePluginPair.Value->Name;
+			OutPackageName = *Plugin->Name;
 			OutLeafPath = VersePathView;
 			return true;
 		}
