Removing superflous error popup, which as a modal dialog blocks the editor (was causing issues when reading .uplugin file in the middle of downloading).

diff --git a/Engine/Source/Runtime/Projects/Private/PluginManager.cpp b/Engine/Source/Runtime/Projects/Private/PluginManager.cpp
--- a/Engine/Source/Runtime/Projects/Private/PluginManager.cpp
+++ b/Engine/Source/Runtime/Projects/Private/PluginManager.cpp
@@ -847,10 +847,7 @@ void FPluginManager::ReadPluginsInDirectory(const FString& PluginsDirectory, con
 			{
 				// NOTE: Even though loading of this plugin failed, we'll keep processing other plugins
 				FString FullPath = FPaths::ConvertRelativePathToFull(FileName);
-				FText FailureMessage = FText::Format(LOCTEXT("FailureFormat", "{0} ({1})"), Context.FailureReason, FText::FromString(FullPath));
-				FText DialogTitle = LOCTEXT("PluginFailureTitle", "Failed to load Plugin");
-				UE_LOG(LogPluginManager, Error, TEXT("%s"), *FailureMessage.ToString());
-				FMessageDialog::Open(EAppMsgType::Ok, FailureMessage, &DialogTitle);
+				UE_LOG(LogPluginManager, Error, TEXT("Failed to load Plugin (%s); %s"), *FullPath, *Context.FailureReason.ToString());
 			}
 		}
 	}
