Always log an error when a plugin fails to load.

#rnx

diff --git a/Engine/Source/Runtime/Projects/Private/PluginManager.cpp b/Engine/Source/Runtime/Projects/Private/PluginManager.cpp
--- a/Engine/Source/Runtime/Projects/Private/PluginManager.cpp
+++ b/Engine/Source/Runtime/Projects/Private/PluginManager.cpp
@@ -2112,11 +2112,8 @@ bool FPluginManager::TryLoadModulesForPlugin( const FPlugin& Plugin, const ELoad
 
 	if( !FailureMessage.IsEmpty() )
 	{
-		if (bAllPluginsEnabledViaCommandLine)
-		{
-			UE_LOG(LogPluginManager, Error, TEXT("%s"), *FailureMessage.ToString());
-		}
-		else
+		UE_LOG(LogPluginManager, Error, TEXT("%s"), *FailureMessage.ToString());
+		if (!bAllPluginsEnabledViaCommandLine)
 		{
 			FMessageDialog::Open(EAppMsgType::Ok, FailureMessage);
 			return false;
