Using 'bEnabled' is not enough to introspect an 'ExplicitlyLoaded' plugin's state. 'ExplicitlyLoaded' plugins can be set enabled in the same (user driven) traditional ways, but not mounted till much later (if ever). Adding an 'IsMounted' state to make the distinction, and fixing a bug preventing 'ExplicitlyLoaded' plugins from being blocked from mounting.

diff --git a/Engine/Source/Runtime/Projects/Private/PluginManager.cpp b/Engine/Source/Runtime/Projects/Private/PluginManager.cpp
--- a/Engine/Source/Runtime/Projects/Private/PluginManager.cpp
+++ b/Engine/Source/Runtime/Projects/Private/PluginManager.cpp
@@ -302,6 +302,7 @@ FPlugin::FPlugin(const FString& InFileName, const FPluginDescriptor& InDescripto
 	, Descriptor(InDescriptor)
 	, Type(InType)
 	, bEnabled(false)
+	, bIsMounted(false)
 {
 
 }
@@ -1616,22 +1617,34 @@ bool FPluginManager::ConfigureEnabledPlugins()
 			}, true); // @todo disable parallelism for now as it's causing hard to track problems
 		}
 
-		for (TSharedRef<IPlugin> Plugin: GetEnabledPluginsWithContentOrVerse())
+		// Mount enabled plugins with content
+		for (const FDiscoveredPluginMap::ElementType& PluginPair : AllPlugins)
 		{
-			if (Plugin->GetDescriptor().bExplicitlyLoaded)
-			{
-				continue;
-			}
+			const TSharedRef<FPlugin>& PluginRef = DiscoveredPluginMapUtils::ResolvePluginFromMapVal(PluginPair.Value);
+			FPlugin& Plugin = *PluginRef;
 
-			if (NewPluginMountedEvent.IsBound())
+			if (Plugin.IsEnabled())
 			{
-				NewPluginMountedEvent.Broadcast(*Plugin);
-			}
+				if (Plugin.GetDescriptor().bExplicitlyLoaded)
+				{
+					continue;
+				}
 
-			if (ensure(RegisterMountPointDelegate.IsBound()))
-			{
-				FString ContentDir = Plugin->GetContentDir();
-				RegisterMountPointDelegate.Execute(Plugin->GetMountedAssetPath(), ContentDir);
+				Plugin.SetIsMounted(true);
+
+				if ((Plugin.CanContainContent() || Plugin.CanContainVerse()))
+				{
+					if (NewPluginMountedEvent.IsBound())
+					{
+						NewPluginMountedEvent.Broadcast(Plugin);
+					}
+
+					if (ensure(RegisterMountPointDelegate.IsBound()))
+					{
+						FString ContentDir = Plugin.GetContentDir();
+						RegisterMountPointDelegate.Execute(Plugin.GetMountedAssetPath(), ContentDir);
+					}
+				}				
 			}
 		}
 
@@ -2468,9 +2481,9 @@ bool FPluginManager::TryMountExplicitlyLoadedPluginVersion(TSharedRef<FPlugin>*
 	if (AllPlugins_PluginPtr && (*AllPlugins_PluginPtr)->Descriptor.bExplicitlyLoaded)
 	{
 		TSharedPtr<FPlugin> CurrentPriorityPlugin = FindPluginInstance((*AllPlugins_PluginPtr)->GetName());
-		if (ensure(CurrentPriorityPlugin.IsValid()) && CurrentPriorityPlugin != *AllPlugins_PluginPtr && CurrentPriorityPlugin->IsEnabled())
+		if (ensure(CurrentPriorityPlugin.IsValid()) && CurrentPriorityPlugin != *AllPlugins_PluginPtr && CurrentPriorityPlugin->IsMounted())
 		{
-			UE_LOG(LogPluginManager, Error, TEXT("Cannot mount plugin '%s' with another version (%s) enabled."), *(*AllPlugins_PluginPtr)->FileName, *CurrentPriorityPlugin->FileName);
+			UE_LOG(LogPluginManager, Error, TEXT("Cannot mount plugin '%s' with another version (%s) already mounted."), *(*AllPlugins_PluginPtr)->FileName, *CurrentPriorityPlugin->FileName);
 		}
 		else
 		{
@@ -2548,6 +2561,8 @@ void FPluginManager::MountPluginFromExternalSource(const TSharedRef<FPlugin>& Pl
 		}
 	}
 
+	Plugin->SetIsMounted(true);
+
 	if (GWarn)
 	{
 		GWarn->EndSlowTask();
@@ -2603,6 +2618,7 @@ bool FPluginManager::UnmountPluginFromExternalSource(const TSharedPtr<FPlugin>&
 		UnRegisterMountPointDelegate.Execute(Plugin->GetMountedAssetPath(), Plugin->GetContentDir());
 	}
 
+	Plugin->SetIsMounted(false);
 	Plugin->bEnabled = false;
 
 	return true;
diff --git a/Engine/Source/Runtime/Projects/Private/PluginManager.h b/Engine/Source/Runtime/Projects/Private/PluginManager.h
--- a/Engine/Source/Runtime/Projects/Private/PluginManager.h
+++ b/Engine/Source/Runtime/Projects/Private/PluginManager.h
@@ -28,7 +28,16 @@ public:
 	EPluginType Type;
 
 	/** True if the plugin is marked as enabled */
-	bool bEnabled;
+	bool bEnabled : 1;
+
+	/** 
+	 * True if this plugin has been mounted/loaded. A 'mounted' plugin is one 
+	 * whose content has been registered, and modules set to load.
+	 * NOTE: Usually plugins are 'enabled' and 'mounted' at the same time, with 
+	 *       the exception of 'ExplicitlyLoaded' plugins (which can be set  
+	 *       'enabled' but not mounted until much later)
+	 */
+	bool bIsMounted : 1;
 
 	/**
 	 * FPlugin constructor
@@ -98,6 +107,9 @@ public:
 #if WITH_EDITOR
 	virtual const TSharedPtr<FJsonObject>& GetDescriptorJson() override;
 #endif // WITH_EDITOR
+
+	bool IsMounted() const { return bIsMounted; }
+	void SetIsMounted(bool bInIsMounted) { bIsMounted = bInIsMounted; }
 };
 
 PRAGMA_DISABLE_DEPRECATION_WARNINGS
