Disabled prescanning of ini files for plugins on desktop platforms - it may end up very slowly scanning Saved/Cooked.

diff --git a/Engine/Source/Runtime/Projects/Private/PluginManager.cpp b/Engine/Source/Runtime/Projects/Private/PluginManager.cpp
--- a/Engine/Source/Runtime/Projects/Private/PluginManager.cpp
+++ b/Engine/Source/Runtime/Projects/Private/PluginManager.cpp
@@ -85,7 +85,7 @@ namespace PluginSystemDefs
 
 	bool IsCachingIniFilesForProcessing()
 	{
-#if PLATFORM_WINDOWS // return false should be false for all platforms according to Josh Adams but has not been tested on other platforms than windows
+#if PLATFORM_DESKTOP // with the reduced set of plugin files to scan, this is likely unnecessary on any platform, but DESKTOP platforms may have Saved/Cooked directories around are slooow to scan
 		return false;
 #endif
 		return true;
