Adding some slow task wrappers around IO tasks which can take a long while (especially when you scale how many plugins your project has).

diff --git a/Engine/Source/Runtime/Projects/Private/PluginManager.cpp b/Engine/Source/Runtime/Projects/Private/PluginManager.cpp
--- a/Engine/Source/Runtime/Projects/Private/PluginManager.cpp
+++ b/Engine/Source/Runtime/Projects/Private/PluginManager.cpp
@@ -651,6 +651,8 @@ void FPluginManager::ReadAllPlugins(FDiscoveredPluginMap& Plugins, const TSet<FS
 	}
 #endif // !WITH_EDITOR
 
+	FScopedSlowTask SlowTask_ReadAll(3.f + ExtraSearchPaths.Num());
+
 	// track child plugins that don't want to go into main plugin set
 	TArray<TSharedRef<FPlugin>> ChildPlugins;
 
@@ -658,30 +660,45 @@ void FPluginManager::ReadAllPlugins(FDiscoveredPluginMap& Plugins, const TSet<FS
 	if (ManifestFileNames.Num() == 0)
 	{
 		UE_LOG(LogPluginManager, Verbose, TEXT("No *.upluginmanifest files found, looking for *.uplugin files instead."))
-		
+
+		SlowTask_ReadAll.EnterProgressFrame(0.5f);
 		// Find "built-in" plugins.  That is, plugins situated right within the Engine directory.
 		TArray<FString> EnginePluginDirs = FPaths::GetExtensionDirs(FPaths::EngineDir(), TEXT("Plugins"));
-		for (const FString& EnginePluginDir : EnginePluginDirs)
 		{
-			ReadPluginsInDirectory(EnginePluginDir, EPluginType::Engine, Plugins, ChildPlugins);
+			FScopedSlowTask SlowTask_ReadEngine((float)EnginePluginDirs.Num());
+			
+			for (const FString& EnginePluginDir : EnginePluginDirs)
+			{
+				SlowTask_ReadEngine.EnterProgressFrame();
+				ReadPluginsInDirectory(EnginePluginDir, EPluginType::Engine, Plugins, ChildPlugins);
+			}
 		}
 
+		SlowTask_ReadAll.EnterProgressFrame(0.5f);
 		// Find plugins in the game project directory (<MyGameProject>/Plugins). If there are any engine plugins matching the name of a game plugin,
 		// assume that the game plugin version is preferred.
 		if (Project != nullptr)
 		{
 			TArray<FString> ProjectPluginDirs = FPaths::GetExtensionDirs(FPaths::GetPath(FPaths::GetProjectFilePath()), TEXT("Plugins"));
+			FScopedSlowTask SlowTask_ReadProject((float)ProjectPluginDirs.Num());
+
 			for (const FString& ProjectPluginDir : ProjectPluginDirs)
 			{
+				SlowTask_ReadProject.EnterProgressFrame();
 				ReadPluginsInDirectory(ProjectPluginDir, EPluginType::Project, Plugins, ChildPlugins);
 			}
 		}
 	}
 	else
 	{
+		SlowTask_ReadAll.EnterProgressFrame();
+		FScopedSlowTask SlowTask_ReadManifest((float)ManifestFileNames.Num());
+
 		// Add plugins from each of the manifests
 		for (const FString& ManifestFileName : ManifestFileNames)
 		{
+			SlowTask_ReadManifest.EnterProgressFrame();
+
 			UE_LOG(LogPluginManager, Verbose, TEXT("Reading plugin manifest: %s"), *ManifestFileName);
 			FPluginManifest Manifest;
 
@@ -698,9 +715,12 @@ void FPluginManager::ReadAllPlugins(FDiscoveredPluginMap& Plugins, const TSet<FS
 			const FString EnterpriseDir = FPaths::EnterpriseDir();
 			const FString ProjectModsDir = FPaths::ProjectModsDir();
 
+			FScopedSlowTask SlowTask_ReadManifestContents((float)Manifest.Contents.Num());
 			// Create all the plugins inside it
 			for (const FPluginManifestEntry& Entry : Manifest.Contents)
 			{
+				SlowTask_ReadManifestContents.EnterProgressFrame();
+
 				EPluginType Type;
 				if (Entry.File.StartsWith(EngineDir) || Entry.File.StartsWith(PlatformExtensionEngineDir))
 				{
@@ -723,17 +743,23 @@ void FPluginManager::ReadAllPlugins(FDiscoveredPluginMap& Plugins, const TSet<FS
 		}
 	}
 
+	SlowTask_ReadAll.EnterProgressFrame();
 	if (Project != nullptr)
 	{
+		FScopedSlowTask SlowTask_ReadProject(2.f + Project->GetAdditionalPluginDirectories().Num());
+
+		SlowTask_ReadProject.EnterProgressFrame();
 		// Always add the mods from the loose directory without using manifests, because they're not packaged together.
 		ReadPluginsInDirectory(FPaths::ProjectModsDir(), EPluginType::Mod, Plugins, ChildPlugins);
 
 		// If they have a list of additional directories to check, add those plugins too
 		for (const FString& Dir : Project->GetAdditionalPluginDirectories())
 		{
+			SlowTask_ReadProject.EnterProgressFrame();
 			ReadPluginsInDirectory(Dir, EPluginType::External, Plugins, ChildPlugins);
 		}
 
+		SlowTask_ReadProject.EnterProgressFrame();
 		// Add plugins from FPaths::EnterprisePluginsDir if it exists
 		if (FPaths::DirectoryExists(FPaths::EnterprisePluginsDir()))
 		{
@@ -743,12 +769,18 @@ void FPluginManager::ReadAllPlugins(FDiscoveredPluginMap& Plugins, const TSet<FS
 
 	for (const FString& ExtraSearchPath : ExtraSearchPaths)
 	{
+		SlowTask_ReadAll.EnterProgressFrame();
 		ReadPluginsInDirectory(ExtraSearchPath, EPluginType::External, Plugins, ChildPlugins);
 	}
 
+	SlowTask_ReadAll.EnterProgressFrame();
+	FScopedSlowTask SlowTask_ReadChildren((float)ChildPlugins.Num());
+
 	// now that we have all the plugins, merge child plugins
 	for (TSharedRef<FPlugin> Child : ChildPlugins)
 	{
+		SlowTask_ReadChildren.EnterProgressFrame();
+
 		// find the parent
 		TArray<FString> Tokens;
 		FPaths::GetCleanFilename(Child->GetDescriptorFileName()).ParseIntoArray(Tokens, TEXT("_"), true);
@@ -812,8 +844,13 @@ void FPluginManager::ReadPluginsInDirectory(const FString& PluginsDirectory, con
 	// Make sure the directory even exists
 	if(FPlatformFileManager::Get().GetPlatformFile().DirectoryExists(*PluginsDirectory))
 	{
+		FScopedSlowTask SlowTask_ReadPlugins(3.f);
+
 		TArray<FString> FileNames;
-		FindPluginsInDirectory(PluginsDirectory, FileNames);
+		{
+			SlowTask_ReadPlugins.EnterProgressFrame();
+			FindPluginsInDirectory(PluginsDirectory, FileNames);
+		}
 
 		struct FLoadContext
 		{
@@ -824,31 +861,41 @@ void FPluginManager::ReadPluginsInDirectory(const FString& PluginsDirectory, con
 
 		TArray<FLoadContext> Contexts;
 		Contexts.SetNum(FileNames.Num());
+		{
+			SlowTask_ReadPlugins.EnterProgressFrame();
 
-		ParallelFor(TEXT("ReadPluginsInDirectory.PF"),
-			FileNames.Num(),1,
-			[&Contexts, &FileNames](int32 Index)
-			{
-				FLoadContext& Context = Contexts[Index];
-				Context.bResult = Context.Descriptor.Load(FileNames[Index], Context.FailureReason);
-			},
-			EParallelForFlags::Unbalanced
-		);
+			ParallelFor(TEXT("ReadPluginsInDirectory.PF"),
+				FileNames.Num(),1,
+				[&Contexts, &FileNames](int32 Index)
+				{
+					FLoadContext& Context = Contexts[Index];
+					Context.bResult = Context.Descriptor.Load(FileNames[Index], Context.FailureReason);
+				},
+				EParallelForFlags::Unbalanced
+			);
+		}
 
-		for (int32 Index = 0, Num = FileNames.Num(); Index < Num; ++Index)
+		SlowTask_ReadPlugins.EnterProgressFrame();
 		{
-			const FString& FileName = FileNames[Index];
-			FLoadContext& Context = Contexts[Index];
-			
-			if (Context.bResult)
-			{
-				CreatePluginObject(FileName, Context.Descriptor, Type, Plugins, ChildPlugins);
-			}
-			else
+			FScopedSlowTask SlowTask_CreatePlugins((float)FileNames.Num());
+
+			for (int32 Index = 0, Num = FileNames.Num(); Index < Num; ++Index)
 			{
-				// NOTE: Even though loading of this plugin failed, we'll keep processing other plugins
-				FString FullPath = FPaths::ConvertRelativePathToFull(FileName);
-				UE_LOG(LogPluginManager, Error, TEXT("Failed to load Plugin (%s); %s"), *FullPath, *Context.FailureReason.ToString());
+				SlowTask_CreatePlugins.EnterProgressFrame();
+
+				const FString& FileName = FileNames[Index];
+				FLoadContext& Context = Contexts[Index];
+
+				if (Context.bResult)
+				{
+					CreatePluginObject(FileName, Context.Descriptor, Type, Plugins, ChildPlugins);
+				}
+				else
+				{
+					// NOTE: Even though loading of this plugin failed, we'll keep processing other plugins
+					FString FullPath = FPaths::ConvertRelativePathToFull(FileName);
+					UE_LOG(LogPluginManager, Error, TEXT("Failed to load Plugin (%s); %s"), *FullPath, *Context.FailureReason.ToString());
+				}
 			}
 		}
 	}
@@ -902,11 +949,17 @@ void FPluginManager::FindPluginsInDirectory(const FString& PluginsDirectory, TAr
 	TArray<FString> DirectoriesToVisit;
 	DirectoriesToVisit.Add(PluginsDirectory);
 
+	FScopedSlowTask SlowTask(1000.f); // Pick an arbitrary amount of work that is resiliant to some floating point multiplication & division
+
 	constexpr int32 MinBatchSize = 1;
 	TArray<TArray<FString>> DirectoriesToVisitNext;
 	FRWLock FoundFilesLock;
 	while (DirectoriesToVisit.Num() > 0)
 	{
+		const float TotalWorkRemaining = SlowTask.TotalAmountOfWork - SlowTask.CompletedWork - SlowTask.CurrentFrameScope;
+		SlowTask.EnterProgressFrame(TotalWorkRemaining);
+		const int32 UnitsOfWorkTodoThisLoop = DirectoriesToVisit.Num();
+
 		ParallelForWithTaskContext(TEXT("FindPluginsInDirectory.PF"),
 			DirectoriesToVisitNext,
 			DirectoriesToVisit.Num(),
@@ -933,13 +986,18 @@ void FPluginManager::FindPluginsInDirectory(const FString& PluginsDirectory, TAr
 				}
 			},
 			EParallelForFlags::Unbalanced);
+		
+		// Adjust the scope of work done this frame, since we discovered more work
+		const int32 NewKnownUnitsOfWork = UnitsOfWorkTodoThisLoop + DirectoriesToVisitNext.Num();
+		SlowTask.CurrentFrameScope = UnitsOfWorkTodoThisLoop * (TotalWorkRemaining / NewKnownUnitsOfWork);
+
 		// Clear and resize `DirectoriesToVisit` for the next batch.
 		DirectoriesToVisit.Reset(Algo::TransformAccumulate(DirectoriesToVisitNext, &TArray<FString>::Num, 0));
 		// Copy all the `DirectoriesToVisitNext` (populated by the various `ParallelFor` tasks) into the one array we use to feed the next round of tasks.
 		for (TArray<FString>& Directories : DirectoriesToVisitNext)
 		{
 			DirectoriesToVisit.Append(MoveTemp(Directories));
-		}
+		}		
 	}
 }
 
@@ -2506,8 +2564,8 @@ bool FPluginManager::TryMountExplicitlyLoadedPluginVersion(TSharedRef<FPlugin>*
 
 				bSuccess = true;
 			}
+		}
 	}
-}
 
 	return bSuccess;
 }
