CIS fixes.

diff --git a/Engine/Source/Runtime/Projects/Private/PluginManager.cpp b/Engine/Source/Runtime/Projects/Private/PluginManager.cpp
--- a/Engine/Source/Runtime/Projects/Private/PluginManager.cpp
+++ b/Engine/Source/Runtime/Projects/Private/PluginManager.cpp
@@ -651,7 +651,7 @@ void FPluginManager::ReadAllPlugins(FDiscoveredPluginMap& Plugins, const TSet<FS
 	}
 #endif // !WITH_EDITOR
 
-	FScopedSlowTask SlowTask_ReadAll(3.f + ExtraSearchPaths.Num());
+	FScopedSlowTask SlowTask_ReadAll(3.f + (float)ExtraSearchPaths.Num());
 
 	// track child plugins that don't want to go into main plugin set
 	TArray<TSharedRef<FPlugin>> ChildPlugins;
@@ -746,7 +746,7 @@ void FPluginManager::ReadAllPlugins(FDiscoveredPluginMap& Plugins, const TSet<FS
 	SlowTask_ReadAll.EnterProgressFrame();
 	if (Project != nullptr)
 	{
-		FScopedSlowTask SlowTask_ReadProject(2.f + Project->GetAdditionalPluginDirectories().Num());
+		FScopedSlowTask SlowTask_ReadProject(2.f + (float)Project->GetAdditionalPluginDirectories().Num());
 
 		SlowTask_ReadProject.EnterProgressFrame();
 		// Always add the mods from the loose directory without using manifests, because they're not packaged together.
@@ -989,7 +989,7 @@ void FPluginManager::FindPluginsInDirectory(const FString& PluginsDirectory, TAr
 		
 		// Adjust the scope of work done this frame, since we discovered more work
 		const int32 NewKnownUnitsOfWork = UnitsOfWorkTodoThisLoop + DirectoriesToVisitNext.Num();
-		SlowTask.CurrentFrameScope = UnitsOfWorkTodoThisLoop * (TotalWorkRemaining / NewKnownUnitsOfWork);
+		SlowTask.CurrentFrameScope = (float)UnitsOfWorkTodoThisLoop * (TotalWorkRemaining / (float)NewKnownUnitsOfWork);
 
 		// Clear and resize `DirectoriesToVisit` for the next batch.
 		DirectoriesToVisit.Reset(Algo::TransformAccumulate(DirectoriesToVisitNext, &TArray<FString>::Num, 0));
