Fix TArray::CheckAddress() error when removing plugin from array.  This was caused by the reference to the plugin shared pointer pointing into the array of plugins.
Fix empty plugin list not being removed from the map.

diff --git a/Engine/Source/Runtime/Projects/Private/PluginManager.cpp b/Engine/Source/Runtime/Projects/Private/PluginManager.cpp
--- a/Engine/Source/Runtime/Projects/Private/PluginManager.cpp
+++ b/Engine/Source/Runtime/Projects/Private/PluginManager.cpp
@@ -256,10 +256,15 @@ namespace DiscoveredPluginMapUtils
 
 	static void RemovePluginFromMap(FDiscoveredPluginMap& InMap, const TSharedRef<FPlugin>& Plugin)
 	{
-		FDiscoveredPluginMap::ValueType* PluginList = InMap.Find(Plugin->Name);
-		if (PluginList && !PluginList->IsEmpty())
+		const uint32 PluginNameHash = GetTypeHash(Plugin->Name);
+		FDiscoveredPluginMap::ValueType* PluginList = InMap.FindByHash(PluginNameHash, Plugin->Name);
+		if (PluginList)
 		{
 			PluginList->Remove(Plugin);
+			if (PluginList->IsEmpty())
+			{
+				InMap.RemoveByHash(PluginNameHash, Plugin->Name);
+			}
 		}
 	}
 
@@ -461,7 +466,7 @@ void FPluginManager::RefreshPluginsList()
 		{
 			const uint32 PluginNameHash = GetTypeHash(Plugin->GetName());
 			PluginsToConfigure.AddByHash(PluginNameHash, Plugin->GetName());
-	}
+		}
 
 #if WITH_EDITOR
 		AddToModuleNameToPluginMap(Plugin);
@@ -581,7 +586,7 @@ bool FPluginManager::RemoveFromPluginsList(const FString& PluginFilename, FText*
 		return true;
 	}
 
-	const TSharedRef<FPlugin>& FoundPlugin = *MaybePlugin;
+	const TSharedRef<FPlugin> FoundPlugin = *MaybePlugin;
 
 	FText FailReason;
 	bool bSuccess = true;
