Found another place where we're conflating enabled with mounted (for explicitly-loaded plugins, enabled does not equate to mounted).

When refershing the plugins list, we start by clearing unused plugin enteries. This change does this for unmounted (explicitly-loaded) plugins too.

diff --git a/Engine/Source/Runtime/Projects/Private/PluginManager.cpp b/Engine/Source/Runtime/Projects/Private/PluginManager.cpp
--- a/Engine/Source/Runtime/Projects/Private/PluginManager.cpp
+++ b/Engine/Source/Runtime/Projects/Private/PluginManager.cpp
@@ -440,14 +440,15 @@ void FPluginManager::RefreshPluginsList()
 	for (FDiscoveredPluginMap::TIterator Iter(AllPlugins); Iter; ++Iter)
 	{
 		const TSharedRef<FPlugin>& Plugin = DiscoveredPluginMapUtils::ResolvePluginFromMapVal(Iter.Value());
-		if(Plugin->bEnabled)
+
+		if (!Plugin->bEnabled || (Plugin->GetDescriptor().bExplicitlyLoaded && !Plugin->bIsMounted))
 		{
-			// Forget all the other discovered versions (which we assume aren't enabled)
-			DiscoveredPluginMapUtils::DiscardAllSupressedVersions(Iter.Value());
+			Iter.RemoveCurrent();
 		}
 		else
 		{
-			Iter.RemoveCurrent();
+			// Forget all the other discovered versions (which we assume aren't enabled/mounted)
+			DiscoveredPluginMapUtils::DiscardAllSupressedVersions(Iter.Value());
 		}
 	}
 
