Add support for environment variables for Virtual Production customers using Unreal in film pipelines.

diff --git a/Engine/Source/Runtime/Projects/Private/PluginManager.cpp b/Engine/Source/Runtime/Projects/Private/PluginManager.cpp
--- a/Engine/Source/Runtime/Projects/Private/PluginManager.cpp
+++ b/Engine/Source/Runtime/Projects/Private/PluginManager.cpp
@@ -39,6 +39,39 @@ DEFINE_LOG_CATEGORY_STATIC( LogPluginManager, Log, All );
 	#define UE_DISABLE_PLUGIN_DISCOVERY 0
 #endif
 
+namespace UE::PluginManager::Private
+{
+TArray<FString> GetPluginPathsByEnv( const TCHAR* EnvVariable )
+{
+	TArray<FString> AdditionalPaths;
+	FString PathStr = FPlatformMisc::GetEnvironmentVariable( EnvVariable );
+	if (!PathStr.IsEmpty())
+	{
+#if PLATFORM_WINDOWS
+		const TCHAR* DirectorySeparator = TEXT(";");
+#else
+		const TCHAR* DirectorySeparator = TEXT(":");
+#endif
+		PathStr.ParseIntoArray(AdditionalPaths, DirectorySeparator, true);
+	}
+	return AdditionalPaths;
+}
+
+/**
+ * Allow users to define additional locations to find plugins. This is to
+ * support traditional DCC film pipelines where plugins can be staged
+ * depending on the context. This is for only looking up compiled plugins
+ * like marketplace plugins and does not imply build support which happens
+ * externally.
+ */
+TArray<FString> GetAdditionalExternalPluginsByEnvVar()
+{
+	return GetPluginPathsByEnv(TEXT("UE_ADDITIONAL_PLUGIN_PATHS"));
+}
+
+}
+
+
 namespace PluginSystemDefs
 {
 	/** File extension of plugin descriptor files.
@@ -80,6 +113,11 @@ namespace PluginSystemDefs
 			}
 		} while (SearchStr != nullptr);
 
+		TArray<FString> AdditionalEnvPaths = UE::PluginManager::Private::GetAdditionalExternalPluginsByEnvVar();
+		for (const FString& Path : AdditionalEnvPaths)
+		{
+			PluginPathsOut.Add(Path);
+		}
 		return PluginCount;
 	}
 
