New plugin descriptor field "EnableVerseAssetReflection"

#rnx

diff --git a/Engine/Plugins/Developer/PluginUtils/Source/PluginUtils/Private/PluginUtils.cpp b/Engine/Plugins/Developer/PluginUtils/Source/PluginUtils/Private/PluginUtils.cpp
--- a/Engine/Plugins/Developer/PluginUtils/Source/PluginUtils/Private/PluginUtils.cpp
+++ b/Engine/Plugins/Developer/PluginUtils/Source/PluginUtils/Private/PluginUtils.cpp
@@ -503,6 +503,7 @@ TSharedPtr<IPlugin> FPluginUtils::CreateAndLoadNewPlugin(const FString& PluginNa
 	ExCreationParams.Descriptor.EnabledByDefault = CreationParams.EnabledByDefault;
 	ExCreationParams.Descriptor.bExplicitlyLoaded = CreationParams.bExplicitelyLoaded;
 	ExCreationParams.Descriptor.VersePath = CreationParams.VersePath;
+	ExCreationParams.Descriptor.bEnableVerseAssetReflection = CreationParams.bEnableVerseAssetReflection;
 
 	if (CreationParams.bHasModules)
 	{
diff --git a/Engine/Plugins/Developer/PluginUtils/Source/PluginUtils/Public/PluginUtils.h b/Engine/Plugins/Developer/PluginUtils/Source/PluginUtils/Public/PluginUtils.h
--- a/Engine/Plugins/Developer/PluginUtils/Source/PluginUtils/Public/PluginUtils.h
+++ b/Engine/Plugins/Developer/PluginUtils/Source/PluginUtils/Public/PluginUtils.h
@@ -92,6 +92,9 @@ public:
 		/** The Verse path to the root of this plugin's content directory */
 		FString VersePath;
 
+		/** If to generate Verse source code definitions from assets contained in this plugin */
+		bool bEnableVerseAssetReflection = false;
+
 		/** Whether this plugin should be enabled/disabled by default for any project. */
 		EPluginEnabledByDefault EnabledByDefault = EPluginEnabledByDefault::Unspecified;
 
diff --git a/Engine/Source/Runtime/Projects/Private/PluginDescriptor.cpp b/Engine/Source/Runtime/Projects/Private/PluginDescriptor.cpp
--- a/Engine/Source/Runtime/Projects/Private/PluginDescriptor.cpp
+++ b/Engine/Source/Runtime/Projects/Private/PluginDescriptor.cpp
@@ -231,6 +231,8 @@ bool FPluginDescriptor::Read(const FJsonObject& Object, FText* OutFailReason /*=
 		}
 	}
 
+	Object.TryGetBoolField(TEXT("EnableVerseAssetReflection"), bEnableVerseAssetReflection);
+
 	bool bEnabledByDefault;
 	if(Object.TryGetBoolField(TEXT("EnabledByDefault"), bEnabledByDefault))
 	{
@@ -351,6 +353,11 @@ void FPluginDescriptor::UpdateJson(FJsonObject& JsonObject) const
 		JsonObject.SetStringField(TEXT("VerseScope"), EVerseScope::ToString(VerseScope));
 	}
 
+	if (bEnableVerseAssetReflection)
+	{
+		JsonObject.SetBoolField(TEXT("EnableVerseAssetReflection"), bEnableVerseAssetReflection);
+	}
+
 	if (EnabledByDefault != EPluginEnabledByDefault::Unspecified)
 	{
 		JsonObject.SetBoolField(TEXT("EnabledByDefault"), (EnabledByDefault == EPluginEnabledByDefault::Enabled));
diff --git a/Engine/Source/Runtime/Projects/Public/PluginDescriptor.h b/Engine/Source/Runtime/Projects/Public/PluginDescriptor.h
--- a/Engine/Source/Runtime/Projects/Public/PluginDescriptor.h
+++ b/Engine/Source/Runtime/Projects/Public/PluginDescriptor.h
@@ -95,6 +95,9 @@ struct PROJECTS_API FPluginDescriptor
 	/** Origin/visibility of Verse code in this plugin's Content/Verse folder */
 	EVerseScope::Type VerseScope = EVerseScope::User;
 
+	/** If to generate Verse source code definitions from assets contained in this plugin */
+	bool bEnableVerseAssetReflection = false;
+
 	/** Whether this plugin should be enabled by default for all projects */
 	EPluginEnabledByDefault EnabledByDefault;
 
