Added a function to the PluginManager to get dependencies for a given FPluginReferenceDescriptor

diff --git a/Engine/Source/Runtime/Projects/Private/PluginManager.cpp b/Engine/Source/Runtime/Projects/Private/PluginManager.cpp
--- a/Engine/Source/Runtime/Projects/Private/PluginManager.cpp
+++ b/Engine/Source/Runtime/Projects/Private/PluginManager.cpp
@@ -2701,6 +2701,39 @@ bool FPluginManager::UnmountExplicitlyLoadedPlugin(const FString& PluginName, FT
 	return UnmountPluginFromExternalSource(Plugin, OutReason);
 }
 
+bool FPluginManager::GetPluginDependencies(const FString& PluginName, TArray<FPluginReferenceDescriptor>& PluginDependencies)
+{
+	const TSharedPtr<FPlugin> Plugin = FindPluginInstance(PluginName);
+	if (Plugin.IsValid())
+	{
+		PluginDependencies = Plugin->Descriptor.Plugins;
+		return true;
+	}
+	return false;
+}
+
+bool FPluginManager::GetPluginDependencies_FromFileName(const FString& PluginFileName, TArray<FPluginReferenceDescriptor>& PluginDependencies)
+{
+	const TSharedRef<FPlugin>* DescribedPlugin = DiscoveredPluginMapUtils::FindPluginInMap_FromFileName(AllPlugins, PluginFileName);
+	if (DescribedPlugin)
+	{
+		PluginDependencies = (*DescribedPlugin)->Descriptor.Plugins;
+		return true;
+	}
+	return false;
+}
+
+bool FPluginManager::GetPluginDependencies_FromDescriptor(const FPluginReferenceDescriptor& PluginDescriptor, TArray<FPluginReferenceDescriptor>& PluginDependencies)
+{
+	const TSharedRef<FPlugin>* DescribedPlugin = DiscoveredPluginMapUtils::FindPluginInMap_FromDescriptor(AllPlugins, PluginDescriptor);
+	if (DescribedPlugin)
+	{
+		PluginDependencies = (*DescribedPlugin)->Descriptor.Plugins;
+		return true;
+	}
+	return false;
+}
+
 bool FPluginManager::UnmountPluginFromExternalSource(const TSharedPtr<FPlugin>& Plugin, FText* OutReason)
 {
 	TRACE_CPUPROFILER_EVENT_SCOPE(UnmountPluginFromExternalSource);
diff --git a/Engine/Source/Runtime/Projects/Private/PluginManager.h b/Engine/Source/Runtime/Projects/Private/PluginManager.h
--- a/Engine/Source/Runtime/Projects/Private/PluginManager.h
+++ b/Engine/Source/Runtime/Projects/Private/PluginManager.h
@@ -173,6 +173,9 @@ public:
 	virtual bool MountExplicitlyLoadedPlugin_FromFileName(const FString& PluginFileName) override;
 	virtual bool MountExplicitlyLoadedPlugin_FromDescriptor(const FPluginReferenceDescriptor& PluginDescriptor) override;
 	virtual bool UnmountExplicitlyLoadedPlugin(const FString& PluginName, FText* OutReason) override;
+	virtual bool GetPluginDependencies(const FString& PluginName, TArray<FPluginReferenceDescriptor>& PluginDependencies) override;
+	virtual bool GetPluginDependencies_FromFileName(const FString& PluginFileName, TArray<FPluginReferenceDescriptor>& PluginDependencies) override;
+	virtual bool GetPluginDependencies_FromDescriptor(const FPluginReferenceDescriptor& PluginDescriptor, TArray<FPluginReferenceDescriptor>& PluginDependencies) override;
 	virtual FName PackageNameFromModuleName(FName ModuleName) override;
 #if UE_USE_VERSE_PATHS
 	virtual bool TrySplitVersePath(const UE::Core::FVersePath& VersePath, FName& OutPackageName, FString& OutLeafPath) override;
diff --git a/Engine/Source/Runtime/Projects/Public/Interfaces/IPluginManager.h b/Engine/Source/Runtime/Projects/Public/Interfaces/IPluginManager.h
--- a/Engine/Source/Runtime/Projects/Public/Interfaces/IPluginManager.h
+++ b/Engine/Source/Runtime/Projects/Public/Interfaces/IPluginManager.h
@@ -415,6 +415,13 @@ public:
 	 */
 	virtual bool UnmountExplicitlyLoadedPlugin(const FString& PluginName, FText* OutReason) = 0;
 
+	/**
+	 * Tries to get a list of plugin dependencies for a given plugin. Returns false if the plugin provided was not found
+	 */
+	virtual bool GetPluginDependencies(const FString& PluginName, TArray<FPluginReferenceDescriptor>& PluginDependencies) = 0;
+	virtual bool GetPluginDependencies_FromFileName(const FString& PluginFileName, TArray<FPluginReferenceDescriptor>& PluginDependencies) = 0;
+	virtual bool GetPluginDependencies_FromDescriptor(const FPluginReferenceDescriptor& PluginDescriptor, TArray<FPluginReferenceDescriptor>& PluginDependencies) = 0;
+
 	/**
 	* Does a reverse lookup to try to figure out what the UObject package name is for a plugin
 	*/
