Don't trigger plugin module loading when querying which plugins have localization data

GetLocalizationPathsForEnabledPlugins will now avoid loading any plugin modules (from its previous call to ConfigureEnabledPlugins), and any plugins that become enabled between GetLocalizationPathsForEnabledPlugins calls will explicitly notify that their localization data is available (as already happens for explicitly loaded plugins).

#rnx

diff --git a/Engine/Source/Runtime/Launch/Private/LaunchEngineLoop.cpp b/Engine/Source/Runtime/Launch/Private/LaunchEngineLoop.cpp
--- a/Engine/Source/Runtime/Launch/Private/LaunchEngineLoop.cpp
+++ b/Engine/Source/Runtime/Launch/Private/LaunchEngineLoop.cpp
@@ -6535,12 +6535,6 @@ bool FEngineLoop::AppInit( )
 		}
 	}
 
-	// Register the callback that allows the text localization manager to load data for plugins
-	FCoreDelegates::GatherAdditionalLocResPathsCallback.AddLambda([](TArray<FString>& OutLocResPaths)
-	{
-		IPluginManager::Get().GetLocalizationPathsForEnabledPlugins(OutLocResPaths);
-	});
-
 	FEmbeddedCommunication::ForceTick(17);
 
 	PreInitHMDDevice();
diff --git a/Engine/Source/Runtime/Projects/Private/PluginManager.cpp b/Engine/Source/Runtime/Projects/Private/PluginManager.cpp
--- a/Engine/Source/Runtime/Projects/Private/PluginManager.cpp
+++ b/Engine/Source/Runtime/Projects/Private/PluginManager.cpp
@@ -343,10 +343,10 @@ namespace DiscoveredPluginMapUtils
 
 namespace PluginLocalizationUtils
 {
-	void GetLocalizationPathsForPlugin(const TSharedRef<FPlugin>& Plugin, TArray<FString>& OutLocResPaths)
+	void GetLocalizationPathsForPlugin(const IPlugin& Plugin, TArray<FString>& OutLocResPaths)
 	{
-		const FString PluginLocDir = Plugin->GetContentDir() / TEXT("Localization");
-		for (const FLocalizationTargetDescriptor& LocTargetDesc : Plugin->GetDescriptor().LocalizationTargets)
+		const FString PluginLocDir = Plugin.GetContentDir() / TEXT("Localization");
+		for (const FLocalizationTargetDescriptor& LocTargetDesc : Plugin.GetDescriptor().LocalizationTargets)
 		{
 			if (LocTargetDesc.ShouldLoadLocalizationTarget())
 			{
@@ -476,6 +476,9 @@ FPluginManager::FPluginManager()
 {
 	SCOPED_BOOT_TIMING("DiscoverAllPlugins");
 	DiscoverAllPlugins();
+
+	// Register the callback that allows the text localization manager to load data for plugins
+	FCoreDelegates::GatherAdditionalLocResPathsCallback.AddRaw(this, &FPluginManager::GetLocalizationPathsForEnabledPlugins);
 }
 
 FPluginManager::~FPluginManager()
@@ -487,6 +490,8 @@ FPluginManager::~FPluginManager()
 	//  shutdown phases.  This will fix issues where modules that are loaded after game modules are shutdown AFTER many engine
 	//  systems are already killed (like GConfig.)  Currently the only workaround is to listen to global exit events, or to
 	//  explicitly unload your module somewhere.  We should be able to handle most cases automatically though!
+
+	FCoreDelegates::GatherAdditionalLocResPathsCallback.RemoveAll(this);
 }
 
 void FPluginManager::RefreshPluginsList()
@@ -1267,6 +1272,9 @@ bool FPluginManager::ConfigureEnabledPlugins()
 		// Keep a set of all the plugin names that have been configured. We read configuration data from different places, but only configure a plugin from the first place that it's referenced.
 		TSet<FString> ConfiguredPluginNames;
 
+		// Keep the list of newly available localization targets
+		TArray<FString> AdditionalLocResPaths;
+
 		// Check which plugins have been enabled or excluded via the command line
 		{
 			SCOPED_BOOT_TIMING("ParseCmdLineForPlugins");
@@ -1528,6 +1536,11 @@ bool FPluginManager::ConfigureEnabledPlugins()
 					continue;
 				}
 #endif
+				if (!Plugin.GetDescriptor().bExplicitlyLoaded && Plugin.GetDescriptor().LocalizationTargets.Num() > 0)
+				{
+					PluginLocalizationUtils::GetLocalizationPathsForPlugin(Plugin, AdditionalLocResPaths);
+				}
+				
 				Plugin.bEnabled = true;
 			}
 		}
@@ -1767,6 +1780,11 @@ bool FPluginManager::ConfigureEnabledPlugins()
 			}
 		}
 
+		if (AdditionalLocResPaths.Num() > 0)
+		{
+			FTextLocalizationManager::Get().HandleLocalizationTargetsMounted(AdditionalLocResPaths);
+		}
+
 		PluginsToConfigure.Empty();
 	}
 	else 
@@ -2309,11 +2327,8 @@ ELoadingPhase::Type FPluginManager::GetLastCompletedLoadingPhase() const
 
 void FPluginManager::GetLocalizationPathsForEnabledPlugins( TArray<FString>& OutLocResPaths )
 {
-	// Figure out which plugins are enabled
-	if (!ConfigureEnabledPlugins())
-	{
-		return;
-	}
+	// Note: We don't call ConfigureEnabledPlugins here as it can cause additional plugin modules to load from a worker thread
+	//       We expect that newly enabled plugins call HandleLocalizationTargetsMounted to load their localization target data
 
 	// Gather the paths from all plugins that have localization targets that are loaded based on the current runtime environment
 	for (const FDiscoveredPluginMap::ElementType& PluginPair : AllPlugins)
@@ -2324,7 +2339,7 @@ void FPluginManager::GetLocalizationPathsForEnabledPlugins( TArray<FString>& Out
 			continue;
 		}
 		
-		PluginLocalizationUtils::GetLocalizationPathsForPlugin(Plugin, OutLocResPaths);
+		PluginLocalizationUtils::GetLocalizationPathsForPlugin(*Plugin, OutLocResPaths);
 	}
 }
 
@@ -2666,7 +2681,7 @@ void FPluginManager::MountPluginFromExternalSource(const TSharedRef<FPlugin>& Pl
 	if (Plugin->Descriptor.LocalizationTargets.Num() > 0)
 	{
 		TArray<FString> AdditionalLocResPaths;
-		PluginLocalizationUtils::GetLocalizationPathsForPlugin(Plugin, AdditionalLocResPaths);
+		PluginLocalizationUtils::GetLocalizationPathsForPlugin(*Plugin, AdditionalLocResPaths);
 		FTextLocalizationManager::Get().HandleLocalizationTargetsMounted(AdditionalLocResPaths);
 	}
 
@@ -2765,7 +2780,7 @@ bool FPluginManager::UnmountPluginFromExternalSource(const TSharedPtr<FPlugin>&
 	if (Plugin->Descriptor.LocalizationTargets.Num() > 0 && !IsEngineExitRequested())
 	{
 		TArray<FString> AdditionalLocResPaths;
-		PluginLocalizationUtils::GetLocalizationPathsForPlugin(Plugin.ToSharedRef(), AdditionalLocResPaths);
+		PluginLocalizationUtils::GetLocalizationPathsForPlugin(*Plugin, AdditionalLocResPaths);
 		FTextLocalizationManager::Get().HandleLocalizationTargetsUnmounted(AdditionalLocResPaths);
 	}
 
