changed PluginManager logging to Display for when the same plugin is prioritized over another
https://p4-swarm.epicgames.net/reviews/23997521

diff --git a/Engine/Source/Runtime/Projects/Private/PluginManager.cpp b/Engine/Source/Runtime/Projects/Private/PluginManager.cpp
--- a/Engine/Source/Runtime/Projects/Private/PluginManager.cpp
+++ b/Engine/Source/Runtime/Projects/Private/PluginManager.cpp
@@ -1118,21 +1118,21 @@ void FPluginManager::CreatePluginObject(const FString& FileName, const FPluginDe
 	// We allow for duplicates of plugins between engine and the project, favoring the project level plugin
 	else if ((*ExistingPlugin)->Type == EPluginType::Engine && Type == EPluginType::Project && !(*ExistingPlugin)->bEnabled)
 	{
-		UE_LOG(LogPluginManager, Verbose, TEXT("By default, prioritizing project plugin (%s) over the corresponding engine version (%s)."), *Plugin->FileName, *(*ExistingPlugin)->FileName);
+		UE_LOG(LogPluginManager, Display, TEXT("By default, prioritizing project plugin (%s) over the corresponding engine version (%s)."), *Plugin->FileName, *(*ExistingPlugin)->FileName);
 		DiscoveredPluginMapUtils::InsertPluginIntoMap(Plugins, Plugin->GetName(), Plugin, DiscoveredPluginMapUtils::EInsertionType::AsOfferedPlugin);
 	}
 	else if ((*ExistingPlugin)->FileName != Plugin->FileName)
 	{
 		if ((*ExistingPlugin)->bEnabled)
 		{
-			UE_LOG(LogPluginManager, Verbose, TEXT("A version of the '%s' plugin has already been enabled (%s); prioritizing that over the newly discovered version (%s)."), *Plugin->Name, *(*ExistingPlugin)->FileName, *Plugin->FileName);
+			UE_LOG(LogPluginManager, Display, TEXT("A version of the '%s' plugin has already been enabled (%s); prioritizing that over the newly discovered version (%s)."), *Plugin->Name, *(*ExistingPlugin)->FileName, *Plugin->FileName);
 			DiscoveredPluginMapUtils::InsertPluginIntoMap(Plugins, Plugin->GetName(), Plugin, DiscoveredPluginMapUtils::EInsertionType::AsSuppressedPlugin);
 		}
 		else if ((*ExistingPlugin)->Type == EPluginType::Project && Type == EPluginType::Engine)
 		{
 			// Project plugins are favored over engine plugins, so we don't want to warn in this case
 			// (instead we mimic the Verbose log from above, and just explain which plugin we're favoring)
-			UE_LOG(LogPluginManager, Verbose, TEXT("By default, prioritizing project plugin (%s) over the corresponding engine version (%s)."), *(*ExistingPlugin)->FileName, *Plugin->FileName);
+			UE_LOG(LogPluginManager, Display, TEXT("By default, prioritizing project plugin (%s) over the corresponding engine version (%s)."), *(*ExistingPlugin)->FileName, *Plugin->FileName);
 			DiscoveredPluginMapUtils::InsertPluginIntoMap(Plugins, Plugin->GetName(), Plugin, DiscoveredPluginMapUtils::EInsertionType::AsSuppressedPlugin);
 		}
 		else
@@ -1142,7 +1142,7 @@ void FPluginManager::CreatePluginObject(const FString& FileName, const FPluginDe
 			{
 				UE_LOG(
 					LogPluginManager,
-					Verbose,
+					Display,
 					TEXT("By default, prioritizing newer version (v%d) of '%s' plugin, over older version (v%d)."),
 					FMath::Max(ExistingVersion, Descriptor.Version),
 					*Plugin->GetName(),
