Expose method to recursively search uplugin files under a specified directory (optimized with parallel directory search and with the assumption that once a uplugin file is found in a specific folder, there cannot be anymore in the same folder hierarchy)

diff --git a/Engine/Source/Runtime/Projects/Private/PluginManager.h b/Engine/Source/Runtime/Projects/Private/PluginManager.h
--- a/Engine/Source/Runtime/Projects/Private/PluginManager.h
+++ b/Engine/Source/Runtime/Projects/Private/PluginManager.h
@@ -145,15 +145,18 @@ public:
 	virtual bool CheckModuleCompatibility(TArray<FString>& OutIncompatibleModules, TArray<FString>& OutIncompatibleEngineModules) override;
 #endif
 	virtual TSharedPtr<IPlugin> FindPlugin(const FStringView Name) override;
-
 	virtual TSharedPtr<IPlugin> FindPlugin(const ANSICHAR* Name) override
 	{
 		FString NameString(Name);
 		return FindPlugin(FStringView(NameString));
 	}
-
 	virtual TSharedPtr<IPlugin> FindPluginFromPath(const FString& PluginPath) override;
 
+	virtual void FindPluginsUnderDirectory(const FString& Directory, TArray<FString>& OutPluginFilePaths) override
+	{
+		return FindPluginsInDirectory(Directory, OutPluginFilePaths);
+	}
+
 	virtual TArray<TSharedRef<IPlugin>> GetEnabledPlugins() override;
 	virtual TArray<TSharedRef<IPlugin>> GetEnabledPluginsWithContent() const override;
 	virtual TArray<TSharedRef<IPlugin>> GetEnabledPluginsWithVerse() const override;
diff --git a/Engine/Source/Runtime/Projects/Public/Interfaces/IPluginManager.h b/Engine/Source/Runtime/Projects/Public/Interfaces/IPluginManager.h
--- a/Engine/Source/Runtime/Projects/Public/Interfaces/IPluginManager.h
+++ b/Engine/Source/Runtime/Projects/Public/Interfaces/IPluginManager.h
@@ -313,6 +313,13 @@ public:
 
 	virtual PROJECTS_API TSharedPtr<IPlugin> FindPluginFromPath(const FString& PluginPath) = 0;
 
+	/** 
+	 * Finds all plugin descriptors underneath a given directory (recursively)
+	 * @param Directory Search folder
+	 * @param OutPluginFilePaths Receives found plugin descriptor file paths
+	 */
+	virtual void FindPluginsUnderDirectory(const FString& Directory, TArray<FString>& OutPluginFilePaths) = 0;
+
 	/**
 	 * Gets an array of all the enabled plugins.
 	 *
