UE-177919 Quick Launch fails when CoTF is enabled due to plugin failing to load modules
-Reverting one part of the change to fix cook on the fly with plugin extensions that add supported target platforms.  The fact that the reference platform list is often empty, and that empty list = all platforms meant this was letting through incorrect plugins in some cases.  Further testing revealed that while the platform extension plugin is rejected here by the previous code it is picked up later on in cook-on-the-fly.

diff --git a/Engine/Source/Runtime/Projects/Private/PluginManager.cpp b/Engine/Source/Runtime/Projects/Private/PluginManager.cpp
--- a/Engine/Source/Runtime/Projects/Private/PluginManager.cpp
+++ b/Engine/Source/Runtime/Projects/Private/PluginManager.cpp
@@ -2047,7 +2047,7 @@ bool FPluginManager::ConfigureEnabledPluginForTarget(const FPluginReferenceDescr
 			}
 
 			// Check the plugin supports this platform
-			if(!bLoadPluginsForTargetPlatforms && !Reference.IsSupportedTargetPlatform(Platform))
+			if (!bLoadPluginsForTargetPlatforms && !Plugin.Descriptor.SupportsTargetPlatform(Platform))
 			{
 				UE_LOG(LogPluginManager, Verbose, TEXT("Ignoring plugin '%s' due to unsupported platform in plugin descriptor"), *Reference.Name);
 				continue;
