Expose finding a plugin via the plugin descritptor. This allows you to find a particular version of a plugin.

diff --git a/Engine/Source/Runtime/Projects/Private/PluginManager.cpp b/Engine/Source/Runtime/Projects/Private/PluginManager.cpp
--- a/Engine/Source/Runtime/Projects/Private/PluginManager.cpp
+++ b/Engine/Source/Runtime/Projects/Private/PluginManager.cpp
@@ -2445,6 +2445,19 @@ TSharedPtr<IPlugin> FPluginManager::FindPluginFromPath(const FString& PluginPath
 	return FindPlugin(PluginName);
 }
 
+TSharedPtr<IPlugin> FPluginManager::FindPluginFromDescriptor(const FPluginReferenceDescriptor& PluginDesc)
+{
+	const TSharedRef<FPlugin>* Instance = DiscoveredPluginMapUtils::FindPluginInMap_FromDescriptor(AllPlugins, PluginDesc);
+	if (Instance == nullptr)
+	{
+		return TSharedPtr<IPlugin>();
+	}
+	else
+	{
+		return TSharedPtr<IPlugin>(*Instance);
+	}
+}
+
 TArray<TSharedRef<IPlugin>> FPluginManager::GetEnabledPlugins()
 {
 	TArray<TSharedRef<IPlugin>> Plugins;
diff --git a/Engine/Source/Runtime/Projects/Private/PluginManager.h b/Engine/Source/Runtime/Projects/Private/PluginManager.h
--- a/Engine/Source/Runtime/Projects/Private/PluginManager.h
+++ b/Engine/Source/Runtime/Projects/Private/PluginManager.h
@@ -151,6 +151,7 @@ public:
 		return FindPlugin(FStringView(NameString));
 	}
 	virtual TSharedPtr<IPlugin> FindPluginFromPath(const FString& PluginPath) override;
+	virtual TSharedPtr<IPlugin> FindPluginFromDescriptor(const FPluginReferenceDescriptor& PluginDesc) override;
 
 	virtual void FindPluginsUnderDirectory(const FString& Directory, TArray<FString>& OutPluginFilePaths) override
 	{
diff --git a/Engine/Source/Runtime/Projects/Public/Interfaces/IPluginManager.h b/Engine/Source/Runtime/Projects/Public/Interfaces/IPluginManager.h
--- a/Engine/Source/Runtime/Projects/Public/Interfaces/IPluginManager.h
+++ b/Engine/Source/Runtime/Projects/Public/Interfaces/IPluginManager.h
@@ -312,6 +312,7 @@ public:
 	virtual PROJECTS_API TSharedPtr<IPlugin> FindPlugin(const ANSICHAR* Name) = 0;
 
 	virtual PROJECTS_API TSharedPtr<IPlugin> FindPluginFromPath(const FString& PluginPath) = 0;
+	virtual PROJECTS_API TSharedPtr<IPlugin> FindPluginFromDescriptor(const FPluginReferenceDescriptor& PluginDesc) = 0;
 
 	/** 
 	 * Finds all plugin descriptors underneath a given directory (recursively)
