code cleanup removed PRAGMA_DISABLE_DEPRECATION_WARNINGS
https://p4-swarm.epicgames.net/reviews/24920073

diff --git a/Engine/Source/Runtime/Projects/Private/PluginManager.h b/Engine/Source/Runtime/Projects/Private/PluginManager.h
--- a/Engine/Source/Runtime/Projects/Private/PluginManager.h
+++ b/Engine/Source/Runtime/Projects/Private/PluginManager.h
@@ -116,7 +116,6 @@ public:
 	void SetIsMounted(bool bInIsMounted) { bIsMounted = bInIsMounted; }
 };
 
-PRAGMA_DISABLE_DEPRECATION_WARNINGS
 /**
  * FPluginManager manages available code and content extensions (both loaded and not loaded.)
  */
@@ -319,6 +318,5 @@ private:
 	/** The highest LoadingPhase that has so far completed */
 	ELoadingPhase::Type LastCompletedLoadingPhase = ELoadingPhase::None;
 };
-PRAGMA_ENABLE_DEPRECATION_WARNINGS
 
 
