code cleanup removed uncessessary declspec on interface functions
https://p4-swarm.epicgames.net/reviews/24920073

diff --git a/Engine/Source/Runtime/Projects/Public/Interfaces/IPluginManager.h b/Engine/Source/Runtime/Projects/Public/Interfaces/IPluginManager.h
--- a/Engine/Source/Runtime/Projects/Public/Interfaces/IPluginManager.h
+++ b/Engine/Source/Runtime/Projects/Public/Interfaces/IPluginManager.h
@@ -308,11 +308,11 @@ public:
 	 *
 	 * @return	 Pointer to the plugin's information, or nullptr.
 	 */
-	virtual PROJECTS_API TSharedPtr<IPlugin> FindPlugin(const FStringView Name) = 0;
-	virtual PROJECTS_API TSharedPtr<IPlugin> FindPlugin(const ANSICHAR* Name) = 0;
+	virtual TSharedPtr<IPlugin> FindPlugin(const FStringView Name) = 0;
+	virtual TSharedPtr<IPlugin> FindPlugin(const ANSICHAR* Name) = 0;
 
-	virtual PROJECTS_API TSharedPtr<IPlugin> FindPluginFromPath(const FString& PluginPath) = 0;
-	virtual PROJECTS_API TSharedPtr<IPlugin> FindPluginFromDescriptor(const FPluginReferenceDescriptor& PluginDesc) = 0;
+	virtual TSharedPtr<IPlugin> FindPluginFromPath(const FString& PluginPath) = 0;
+	virtual TSharedPtr<IPlugin> FindPluginFromDescriptor(const FPluginReferenceDescriptor& PluginDesc) = 0;
 
 	/** 
 	 * Finds all plugin descriptors underneath a given directory (recursively)
@@ -358,12 +358,12 @@ public:
 	/**
 	 * Returns the set of built-in plugin names
 	 */
-	PROJECTS_API virtual const TSet<FString>& GetBuiltInPluginNames() const = 0;
+	virtual const TSet<FString>& GetBuiltInPluginNames() const = 0;
 
 	/**
 	 * Returns the plugin that owns the specified module, if any
 	 */
-	PROJECTS_API virtual TSharedPtr<IPlugin> GetModuleOwnerPlugin(FName ModuleName) const = 0;
+	virtual TSharedPtr<IPlugin> GetModuleOwnerPlugin(FName ModuleName) const = 0;
 #endif //WITH_EDITOR
 
 	/**
diff --git a/Engine/Source/Runtime/Projects/Public/Interfaces/IProjectManager.h b/Engine/Source/Runtime/Projects/Public/Interfaces/IProjectManager.h
--- a/Engine/Source/Runtime/Projects/Public/Interfaces/IProjectManager.h
+++ b/Engine/Source/Runtime/Projects/Public/Interfaces/IProjectManager.h
@@ -243,7 +243,7 @@ public:
 	 *
 	 * @return	True if the project is an Enterprise project
 	 */
-	PROJECTS_API virtual bool IsEnterpriseProject() = 0;
+	virtual bool IsEnterpriseProject() = 0;
 
 	/**
 	 * Sets the enterprise flag value on the current project
@@ -261,8 +261,8 @@ public:
 	virtual bool IsSuppressingProjectFileWrite() const = 0;
 
 	/** Suppress project file writes. */
-	PROJECTS_API virtual void AddSuppressProjectFileWrite(const FName InName) = 0;
+	virtual void AddSuppressProjectFileWrite(const FName InName) = 0;
 
 	/** Removes suppression of project file writes. */
-	PROJECTS_API virtual void RemoveSuppressProjectFileWrite(const FName InName) = 0;
+	virtual void RemoveSuppressProjectFileWrite(const FName InName) = 0;
 };
