Plugin Sealed load time validation - Prompt the user to disable the plugin or close the application if a plugin is found to have a dependency on a sealed plugin

diff --git a/Engine/Source/Runtime/Projects/Private/PluginManager.cpp b/Engine/Source/Runtime/Projects/Private/PluginManager.cpp
--- a/Engine/Source/Runtime/Projects/Private/PluginManager.cpp
+++ b/Engine/Source/Runtime/Projects/Private/PluginManager.cpp
@@ -1866,17 +1866,18 @@ bool FPluginManager::ConfigureEnabledPlugins()
 
 bool FPluginManager::RequiresTempTargetForCodePlugin(const FProjectDescriptor* ProjectDescriptor, const FString& Platform, EBuildConfiguration Configuration, EBuildTargetType TargetType, FText& OutReason)
 {
-	const FPluginReferenceDescriptor* MissingPlugin;
+	const FPluginReferenceDescriptor* MissingPlugin = nullptr;
+	const FPluginReferenceDescriptor* SealedPlugin = nullptr;
 
 	TSet<FString> ProjectCodePlugins;
-	if (!GetCodePluginsForProject(ProjectDescriptor, Platform, Configuration, TargetType, AllPlugins, ProjectCodePlugins, MissingPlugin))
+	if (!GetCodePluginsForProject(ProjectDescriptor, Platform, Configuration, TargetType, AllPlugins, ProjectCodePlugins, MissingPlugin, SealedPlugin))
 	{
 		OutReason = FText::Format(LOCTEXT("TempTarget_MissingPluginForTarget", "{0} plugin is referenced by target but not found"), FText::FromString(MissingPlugin->Name));
 		return true;
 	}
 
 	TSet<FString> DefaultCodePlugins;
-	if (!GetCodePluginsForProject(nullptr, Platform, Configuration, TargetType, AllPlugins, DefaultCodePlugins, MissingPlugin))
+	if (!GetCodePluginsForProject(nullptr, Platform, Configuration, TargetType, AllPlugins, DefaultCodePlugins, MissingPlugin, SealedPlugin))
 	{
 		OutReason = FText::Format(LOCTEXT("TempTarget_MissingPluginForDefaultTarget", "{0} plugin is referenced by the default target but not found"), FText::FromString(MissingPlugin->Name));
 		return true;
@@ -1903,7 +1904,7 @@ bool FPluginManager::RequiresTempTargetForCodePlugin(const FProjectDescriptor* P
 	return false;
 }
 
-bool FPluginManager::GetCodePluginsForProject(const FProjectDescriptor* ProjectDescriptor, const FString& Platform, EBuildConfiguration Configuration, EBuildTargetType TargetType, FDiscoveredPluginMap& AllPlugins, TSet<FString>& CodePluginNames, const FPluginReferenceDescriptor*& OutMissingPlugin)
+bool FPluginManager::GetCodePluginsForProject(const FProjectDescriptor* ProjectDescriptor, const FString& Platform, EBuildConfiguration Configuration, EBuildTargetType TargetType, FDiscoveredPluginMap& AllPlugins, TSet<FString>& CodePluginNames, const FPluginReferenceDescriptor*& OutMissingPlugin, const FPluginReferenceDescriptor*& OutSealedPlugin)
 {
 	// Can only check the current project at the moment, since we won't have enumerated them otherwise
 	check(ProjectDescriptor == nullptr || ProjectDescriptor == IProjectManager::Get().GetCurrentProject());
@@ -1929,7 +1930,7 @@ bool FPluginManager::GetCodePluginsForProject(const FProjectDescriptor* ProjectD
 		{
 			if(!ConfiguredPluginNames.Contains(PluginReference.Name))
 			{
-				if (!ConfigureEnabledPluginForTarget(PluginReference, ProjectDescriptor, FString(), Platform, Configuration, TargetType, bLoadPluginsForTargetPlatforms, AllPlugins, EnabledPlugins, OutMissingPlugin))
+				if (!ConfigureEnabledPluginForTarget(PluginReference, ProjectDescriptor, FString(), Platform, Configuration, TargetType, bLoadPluginsForTargetPlatforms, AllPlugins, EnabledPlugins, OutMissingPlugin, OutSealedPlugin))
 				{
 					return false;
 				}
@@ -1943,7 +1944,7 @@ bool FPluginManager::GetCodePluginsForProject(const FProjectDescriptor* ProjectD
 	{
 		if (DiscoveredPluginMapUtils::ResolvePluginFromMapVal(PluginPair.Value)->IsEnabledByDefault(bAllowEnginePluginsEnabledByDefault) && !ConfiguredPluginNames.Contains(PluginPair.Key))
 		{
-			if (!ConfigureEnabledPluginForTarget(FPluginReferenceDescriptor(PluginPair.Key, true), ProjectDescriptor, FString(), Platform, Configuration, TargetType, bLoadPluginsForTargetPlatforms, AllPlugins, EnabledPlugins, OutMissingPlugin))
+			if (!ConfigureEnabledPluginForTarget(FPluginReferenceDescriptor(PluginPair.Key, true), ProjectDescriptor, FString(), Platform, Configuration, TargetType, bLoadPluginsForTargetPlatforms, AllPlugins, EnabledPlugins, OutMissingPlugin, OutSealedPlugin))
 			{
 				return false;
 			}
@@ -1973,8 +1974,9 @@ bool FPluginManager::ConfigureEnabledPluginForCurrentTarget(const FPluginReferen
 {
 	SCOPED_BOOT_TIMING("ConfigureEnabledPluginForCurrentTarget");
 
-	const FPluginReferenceDescriptor* MissingPlugin;
-	if (!ConfigureEnabledPluginForTarget(FirstReference, IProjectManager::Get().GetCurrentProject(), UE_APP_NAME, FPlatformMisc::GetUBTPlatform(), FApp::GetBuildConfiguration(), FApp::GetBuildTargetType(), (bool)LOAD_PLUGINS_FOR_TARGET_PLATFORMS, AllPlugins, EnabledPlugins, MissingPlugin))
+	const FPluginReferenceDescriptor* MissingPlugin = nullptr;
+	const FPluginReferenceDescriptor* SealedPlugin = nullptr;
+	if (!ConfigureEnabledPluginForTarget(FirstReference, IProjectManager::Get().GetCurrentProject(), UE_APP_NAME, FPlatformMisc::GetUBTPlatform(), FApp::GetBuildConfiguration(), FApp::GetBuildTargetType(), (bool)LOAD_PLUGINS_FOR_TARGET_PLATFORMS, AllPlugins, EnabledPlugins, MissingPlugin, SealedPlugin))
 	{
 		// If we're in unattended mode, don't open any windows and fatal out
 		if (FApp::IsUnattended())
@@ -1984,29 +1986,41 @@ bool FPluginManager::ConfigureEnabledPluginForCurrentTarget(const FPluginReferen
 		}
 
 #if !IS_MONOLITHIC
-		// Try to download it from the marketplace
-		if (MissingPlugin->MarketplaceURL.Len() > 0 && PromptToDownloadPlugin(MissingPlugin->Name, MissingPlugin->MarketplaceURL))
+		if (MissingPlugin != nullptr)
 		{
-			UE_LOG(LogPluginManager, Display, TEXT("Downloading '%s' plugin from marketplace (%s)."), *MissingPlugin->Name, *MissingPlugin->MarketplaceURL);
-			return false;
-		}
+			// Try to download it from the marketplace
+			if (MissingPlugin->MarketplaceURL.Len() > 0 && PromptToDownloadPlugin(MissingPlugin->Name, MissingPlugin->MarketplaceURL))
+			{
+				UE_LOG(LogPluginManager, Display, TEXT("Downloading '%s' plugin from marketplace (%s)."), *MissingPlugin->Name, *MissingPlugin->MarketplaceURL);
+				return false;
+			}
 
-		// Prompt to disable it in the project file, if possible
-		if (PromptToDisableMissingPlugin(FirstReference.Name, MissingPlugin->Name))
+			// Prompt to disable it in the project file, if possible
+			if (PromptToDisableMissingPlugin(FirstReference.Name, MissingPlugin->Name))
+			{
+				UE_LOG(LogPluginManager, Display, TEXT("Disabled plugin '%s', continuing."), *FirstReference.Name);
+				return true;
+			}
+		}
+		else if (SealedPlugin != nullptr)
 		{
-			UE_LOG(LogPluginManager, Display, TEXT("Disabled plugin '%s', continuing."), *FirstReference.Name);
-			return true;
+			// Prompt to disable it in the project file, if possible
+			if (PromptToDisableSealedPlugin(FirstReference.Name, SealedPlugin->Name))
+			{
+				UE_LOG(LogPluginManager, Display, TEXT("Disabled plugin '%s', continuing."), *FirstReference.Name);
+				return true;
+			}
 		}
 #endif
 
 		// Unable to continue
-		UE_LOG(LogPluginManager, Error, TEXT("Unable to load plugin '%s'. Aborting."), *MissingPlugin->Name);
+		UE_LOG(LogPluginManager, Error, TEXT("Unable to load plugin '%s'. Aborting."), MissingPlugin ? *MissingPlugin->Name : *SealedPlugin->Name);
 		return false;
 	}
 	return true;
 }
 
-bool FPluginManager::ConfigureEnabledPluginForTarget(const FPluginReferenceDescriptor& FirstReference, const FProjectDescriptor* ProjectDescriptor, const FString& TargetName, const FString& Platform, EBuildConfiguration Configuration, EBuildTargetType TargetType, bool bLoadPluginsForTargetPlatforms, FDiscoveredPluginMap& AllPlugins, TMap<FString, FPlugin*>& EnabledPlugins, const FPluginReferenceDescriptor*& OutMissingPlugin)
+bool FPluginManager::ConfigureEnabledPluginForTarget(const FPluginReferenceDescriptor& FirstReference, const FProjectDescriptor* ProjectDescriptor, const FString& TargetName, const FString& Platform, EBuildConfiguration Configuration, EBuildTargetType TargetType, bool bLoadPluginsForTargetPlatforms, FDiscoveredPluginMap& AllPlugins, TMap<FString, FPlugin*>& EnabledPlugins, const FPluginReferenceDescriptor*& OutMissingPlugin, const FPluginReferenceDescriptor*& OutSealedPlugin)
 {
 	if (!EnabledPlugins.Contains(FirstReference.Name))
 	{
@@ -2145,6 +2159,13 @@ bool FPluginManager::ConfigureEnabledPluginForTarget(const FPluginReferenceDescr
 			// Add references to all its dependencies
 			for (const FPluginReferenceDescriptor& NextReference : Plugin.Descriptor.Plugins)
 			{
+				const TSharedRef<FPlugin>* NextPluginPtr = DiscoveredPluginMapUtils::FindPluginInMap(AllPlugins, NextReference.Name);
+				if (NextPluginPtr != nullptr && NextPluginPtr->Get().Descriptor.bIsSealed)
+				{
+					OutSealedPlugin = &NextReference;
+					return false;
+				}
+
 				if (!EnabledPlugins.Contains(NextReference.Name) && !NewPluginNames.Contains(NextReference.Name))
 				{
 					NewPluginNames.Add(NextReference.Name);
@@ -2191,6 +2212,14 @@ bool FPluginManager::PromptToDownloadPlugin(const FString& PluginName, const FSt
 	return false;
 }
 
+bool FPluginManager::PromptToDisableSealedPlugin(const FString& PluginName, const FString& SealedPluginName)
+{
+	FText Message = FText::Format(LOCTEXT("DisablePluginMessage_IllegalDependency", "This project requires the '{0}' plugin, which has an illegal dependency on the sealed '{1}' plugin.\n\nWould you like to disable it?\n\nIf you do, you will no longer be able to open any assets created with it. If not, the application will close."), FText::FromString(PluginName), FText::FromString(SealedPluginName));
+	FText Caption(LOCTEXT("DisablePluginCaption", "Illegal Dependency"));
+
+	return PromptToDisablePlugin(Caption, Message, PluginName);
+}
+
 bool FPluginManager::PromptToDisableMissingPlugin(const FString& PluginName, const FString& MissingPluginName)
 {
 	FText Message;
diff --git a/Engine/Source/Runtime/Projects/Private/PluginManager.h b/Engine/Source/Runtime/Projects/Private/PluginManager.h
--- a/Engine/Source/Runtime/Projects/Private/PluginManager.h
+++ b/Engine/Source/Runtime/Projects/Private/PluginManager.h
@@ -215,7 +215,7 @@ private:
 	static void FindPluginManifestsInDirectory(const FString& PluginManifestDirectory, TArray<FString>& FileNames);
 
 	/** Gets all the code plugins that are enabled for a content only project */
-	static bool GetCodePluginsForProject(const FProjectDescriptor* ProjectDescriptor, const FString& Platform, EBuildConfiguration Configuration, EBuildTargetType TargetType, FDiscoveredPluginMap& AllPlugins, TSet<FString>& CodePluginNames, const FPluginReferenceDescriptor*& OutMissingPlugin);
+	static bool GetCodePluginsForProject(const FProjectDescriptor* ProjectDescriptor, const FString& Platform, EBuildConfiguration Configuration, EBuildTargetType TargetType, FDiscoveredPluginMap& AllPlugins, TSet<FString>& CodePluginNames, const FPluginReferenceDescriptor*& OutMissingPlugin, const FPluginReferenceDescriptor*& OutSealedPlugin);
 
 	/** Sets the bPluginEnabled flag on all plugins found from DiscoverAllPlugins that are enabled in config */
 	bool ConfigureEnabledPlugins();
@@ -224,11 +224,14 @@ private:
 	bool ConfigureEnabledPluginForCurrentTarget(const FPluginReferenceDescriptor& FirstReference, TMap<FString, FPlugin*>& EnabledPlugins);
 
 	/** Adds a single enabled plugin and all its dependencies. */
-	static bool ConfigureEnabledPluginForTarget(const FPluginReferenceDescriptor& FirstReference, const FProjectDescriptor* ProjectDescriptor, const FString& TargetName, const FString& Platform, EBuildConfiguration Configuration, EBuildTargetType TargetType, bool bLoadPluginsForTargetPlatforms, FDiscoveredPluginMap& AllPlugins, TMap<FString, FPlugin*>& EnabledPlugins, const FPluginReferenceDescriptor*& OutMissingPlugin);
+	static bool ConfigureEnabledPluginForTarget(const FPluginReferenceDescriptor& FirstReference, const FProjectDescriptor* ProjectDescriptor, const FString& TargetName, const FString& Platform, EBuildConfiguration Configuration, EBuildTargetType TargetType, bool bLoadPluginsForTargetPlatforms, FDiscoveredPluginMap& AllPlugins, TMap<FString, FPlugin*>& EnabledPlugins, const FPluginReferenceDescriptor*& OutMissingPlugin, const FPluginReferenceDescriptor*& OutSealedPlugin);
 
 	/** Prompts the user to download a missing plugin from the given URL */
 	static bool PromptToDownloadPlugin(const FString& PluginName, const FString& MarketplaceURL);
 
+	/** Prompts the user to disable a plugin */
+	static bool PromptToDisableSealedPlugin(const FString& PluginName, const FString& SealedPluginName);
+
 	/** Prompts the user to disable a plugin */
 	static bool PromptToDisableMissingPlugin(const FString& PluginName, const FString& MissingPluginName);
 
