Fixing conflicting loctext. Issue 306621

#rnx

diff --git a/Engine/Source/Runtime/Projects/Private/PluginManager.cpp b/Engine/Source/Runtime/Projects/Private/PluginManager.cpp
--- a/Engine/Source/Runtime/Projects/Private/PluginManager.cpp
+++ b/Engine/Source/Runtime/Projects/Private/PluginManager.cpp
@@ -2215,7 +2215,7 @@ bool FPluginManager::PromptToDownloadPlugin(const FString& PluginName, const FSt
 bool FPluginManager::PromptToDisableSealedPlugin(const FString& PluginName, const FString& SealedPluginName)
 {
 	FText Message = FText::Format(LOCTEXT("DisablePluginMessage_IllegalDependency", "This project requires the '{0}' plugin, which has an illegal dependency on the sealed '{1}' plugin.\n\nWould you like to disable it?\n\nIf you do, you will no longer be able to open any assets created with it. If not, the application will close."), FText::FromString(PluginName), FText::FromString(SealedPluginName));
-	FText Caption(LOCTEXT("DisablePluginCaption", "Illegal Dependency"));
+	FText Caption(LOCTEXT("DisablePluginCaption_IllegalDependency", "Illegal Dependency"));
 
 	return PromptToDisablePlugin(Caption, Message, PluginName);
 }
