Fixing static analysis warning

#rnx

diff --git a/Engine/Source/Runtime/Projects/Private/PluginManager.cpp b/Engine/Source/Runtime/Projects/Private/PluginManager.cpp
--- a/Engine/Source/Runtime/Projects/Private/PluginManager.cpp
+++ b/Engine/Source/Runtime/Projects/Private/PluginManager.cpp
@@ -1990,7 +1990,14 @@ bool FPluginManager::ConfigureEnabledPluginForCurrentTarget(const FPluginReferen
 		// If we're in unattended mode, don't open any windows and fatal out
 		if (FApp::IsUnattended())
 		{
-			UE_LOG(LogPluginManager, Fatal, TEXT("This project requires the '%s' plugin. Install it and try again, or remove it from the project's required plugin list."), *MissingPlugin->Name);
+			if (MissingPlugin != nullptr)
+			{
+				UE_LOG(LogPluginManager, Fatal, TEXT("This project requires the '%s' plugin. Install it and try again, or remove it from the project's required plugin list."), *MissingPlugin->Name);
+			}
+			else if (SealedPlugin != nullptr)
+			{
+				UE_LOG(LogPluginManager, Fatal, TEXT("This project requires the '%s' plugin, which has an illegal dependency on the sealed plugin '%s'."), *FirstReference.Name, *MissingPlugin->Name);
+			}
 			return false;
 		}
 
