Correcting typo in PluginManager.cpp causing static analysis warning's as error

diff --git a/Engine/Source/Runtime/Projects/Private/PluginManager.cpp b/Engine/Source/Runtime/Projects/Private/PluginManager.cpp
--- a/Engine/Source/Runtime/Projects/Private/PluginManager.cpp
+++ b/Engine/Source/Runtime/Projects/Private/PluginManager.cpp
@@ -1987,7 +1987,7 @@ bool FPluginManager::ConfigureEnabledPluginForCurrentTarget(const FPluginReferen
 			}
 			else if (SealedPlugin != nullptr)
 			{
-				UE_LOG(LogPluginManager, Fatal, TEXT("This project requires the '%s' plugin, which has an illegal dependency on the sealed plugin '%s'."), *FirstReference.Name, *MissingPlugin->Name);
+				UE_LOG(LogPluginManager, Fatal, TEXT("This project requires the '%s' plugin, which has an illegal dependency on the sealed plugin '%s'."), *FirstReference.Name, *SealedPlugin->Name);
 			}
 			return false;
 		}
