NoCode plugin descriptor property to prevent a plugin from containing any code or modules

diff --git a/Engine/Plugins/Editor/PluginBrowser/Source/PluginBrowser/Private/PluginMetadataObject.cpp b/Engine/Plugins/Editor/PluginBrowser/Source/PluginBrowser/Private/PluginMetadataObject.cpp
--- a/Engine/Plugins/Editor/PluginBrowser/Source/PluginBrowser/Private/PluginMetadataObject.cpp
+++ b/Engine/Plugins/Editor/PluginBrowser/Source/PluginBrowser/Private/PluginMetadataObject.cpp
@@ -52,6 +52,7 @@ void UPluginMetadataObject::PopulateFromPlugin(TSharedPtr<IPlugin> InPlugin)
 	bIsEnabledByDefault = (InDescriptor.EnabledByDefault == EPluginEnabledByDefault::Enabled);
 	bExplicitlyLoaded = InDescriptor.bExplicitlyLoaded;
 	bIsSealed = InDescriptor.bIsSealed;
+	bNoCode = InDescriptor.bNoCode;
 
 	Plugins.Reset(InDescriptor.Plugins.Num());
 	for (const FPluginReferenceDescriptor& PluginRefDesc : InDescriptor.Plugins)
@@ -77,6 +78,7 @@ void UPluginMetadataObject::CopyIntoDescriptor(FPluginDescriptor& OutDescriptor)
 	OutDescriptor.bCanContainContent = bCanContainContent;
 	OutDescriptor.bIsBetaVersion = bIsBetaVersion;
 	OutDescriptor.bIsSealed = bIsSealed;
+	OutDescriptor.bNoCode = bNoCode;
 
 	TArray<FPluginReferenceDescriptor> NewPlugins;
 	NewPlugins.Reserve(Plugins.Num());
diff --git a/Engine/Plugins/Editor/PluginBrowser/Source/PluginBrowser/Private/PluginMetadataObject.h b/Engine/Plugins/Editor/PluginBrowser/Source/PluginBrowser/Private/PluginMetadataObject.h
--- a/Engine/Plugins/Editor/PluginBrowser/Source/PluginBrowser/Private/PluginMetadataObject.h
+++ b/Engine/Plugins/Editor/PluginBrowser/Source/PluginBrowser/Private/PluginMetadataObject.h
@@ -125,6 +125,9 @@ public:
 	UPROPERTY(EditAnywhere, Category = Details)
 	bool bIsSealed;
 
+	UPROPERTY(EditAnywhere, Category = Details)
+	bool bNoCode;
+
 	/** Plugins used by this plugin */
 	UPROPERTY(EditAnywhere, Category = Dependencies, meta=(TitleProperty=Name))
 	TArray<FPluginReferenceMetadata> Plugins;
diff --git a/Engine/Source/Programs/UnrealBuildTool/Configuration/UEBuildTarget.cs b/Engine/Source/Programs/UnrealBuildTool/Configuration/UEBuildTarget.cs
--- a/Engine/Source/Programs/UnrealBuildTool/Configuration/UEBuildTarget.cs
+++ b/Engine/Source/Programs/UnrealBuildTool/Configuration/UEBuildTarget.cs
@@ -2283,6 +2283,12 @@ namespace UnrealBuildTool
 			// Check that each plugin declares its dependencies explicitly
 			foreach (UEBuildPlugin Plugin in BuildPlugins)
 			{
+				// Check that any plugins with the NoCode specifier do not contain modules.
+				if (Plugin.Descriptor.bNoCode && Plugin.Modules.Count > 0)
+				{
+					throw new BuildException("Plugin '{0}' cannot contain any code or modules. See the plugin descriptor property `NoCode`", Plugin.Name);
+				}
+
 				foreach (UEBuildModule Module in Plugin.Modules)
 				{
 					HashSet<UEBuildModule> DependencyModules = Module.GetDependencies(bWithIncludePathModules: true, bWithDynamicallyLoadedModules: true);
diff --git a/Engine/Source/Programs/UnrealBuildTool/System/PluginDescriptor.cs b/Engine/Source/Programs/UnrealBuildTool/System/PluginDescriptor.cs
--- a/Engine/Source/Programs/UnrealBuildTool/System/PluginDescriptor.cs
+++ b/Engine/Source/Programs/UnrealBuildTool/System/PluginDescriptor.cs
@@ -186,6 +186,11 @@ namespace UnrealBuildTool
 		/// </summary>
 		public bool bIsSealed;
 
+		/// <summary>
+		/// When true, this plugin should not contain any code or modules.
+		/// </summary>
+		public bool bNoCode;
+
 		/// <summary>
 		/// When true, this plugin's modules will not be loaded automatically nor will it's content be mounted automatically. It will load/mount when explicitly requested and LoadingPhases will be ignored
 		/// </summary>
@@ -329,6 +334,7 @@ namespace UnrealBuildTool
 
 			RawObject.TryGetBoolField("RequiresBuildPlatform", out bRequiresBuildPlatform);
 			RawObject.TryGetBoolField("Sealed", out bIsSealed);
+			RawObject.TryGetBoolField("NoCode", out bNoCode);
 			RawObject.TryGetBoolField("ExplicitlyLoaded", out bExplicitlyLoaded);
 			RawObject.TryGetBoolField("HasExplicitPlatforms", out bHasExplicitPlatforms);
 
@@ -437,6 +443,11 @@ namespace UnrealBuildTool
 				Writer.WriteValue("Sealed", bIsSealed);
 			}
 
+			if (bNoCode)
+			{
+				Writer.WriteValue("NoCode", bNoCode);
+			}
+
 			if (bExplicitlyLoaded)
 			{
 				Writer.WriteValue("ExplicitlyLoaded", bExplicitlyLoaded);
diff --git a/Engine/Source/Runtime/Projects/Private/PluginDescriptor.cpp b/Engine/Source/Runtime/Projects/Private/PluginDescriptor.cpp
--- a/Engine/Source/Runtime/Projects/Private/PluginDescriptor.cpp
+++ b/Engine/Source/Runtime/Projects/Private/PluginDescriptor.cpp
@@ -87,6 +87,7 @@ FPluginDescriptor::FPluginDescriptor()
 	, bRequiresBuildPlatform(false)
 	, bIsHidden(false)
 	, bIsSealed(false)
+	, bNoCode(false)
 	, bExplicitlyLoaded(false)
 	, bHasExplicitPlatforms(false)
 	, bIsPluginExtension(false)
@@ -242,6 +243,7 @@ bool FPluginDescriptor::Read(const FJsonObject& Object, FText* OutFailReason /*=
 
 	Object.TryGetBoolField(TEXT("CanContainContent"), bCanContainContent);
 	Object.TryGetBoolField(TEXT("CanContainVerse"), bCanContainVerse);
+	Object.TryGetBoolField(TEXT("NoCode"), bNoCode);
 	Object.TryGetBoolField(TEXT("IsBetaVersion"), bIsBetaVersion);
 	Object.TryGetBoolField(TEXT("IsExperimentalVersion"), bIsExperimentalVersion);
 	Object.TryGetBoolField(TEXT("Installed"), bInstalled);
@@ -374,6 +376,10 @@ void FPluginDescriptor::UpdateJson(FJsonObject& JsonObject) const
 	{
 		JsonObject.SetBoolField(TEXT("CanContainVerse"), bCanContainVerse);
 	}
+	if (bNoCode)
+	{
+		JsonObject.SetBoolField(TEXT("NoCode"), bNoCode);
+	}
 	JsonObject.SetBoolField(TEXT("IsBetaVersion"), bIsBetaVersion);
 	JsonObject.SetBoolField(TEXT("IsExperimentalVersion"), bIsExperimentalVersion);
 	JsonObject.SetBoolField(TEXT("Installed"), bInstalled);
diff --git a/Engine/Source/Runtime/Projects/Public/PluginDescriptor.h b/Engine/Source/Runtime/Projects/Public/PluginDescriptor.h
--- a/Engine/Source/Runtime/Projects/Public/PluginDescriptor.h
+++ b/Engine/Source/Runtime/Projects/Public/PluginDescriptor.h
@@ -125,6 +125,9 @@ struct PROJECTS_API FPluginDescriptor
 	/** Prevents other plugins from depending on this plugin. */
 	bool bIsSealed;
 
+	/** Prevents this plugin from containing code or modules. */
+	bool bNoCode;
+
 	/** When true, this plugin's modules will not be loaded automatically nor will it's content be mounted automatically. It will load/mount when explicitly requested and LoadingPhases will be ignored */
 	bool bExplicitlyLoaded;
 
