Added a new DisallowedPlugins property to the PluginDescriptor that restricts which plugins can be used by this plugin.

diff --git a/Engine/Plugins/Editor/PluginBrowser/Source/PluginBrowser/Private/PluginMetadataObject.cpp b/Engine/Plugins/Editor/PluginBrowser/Source/PluginBrowser/Private/PluginMetadataObject.cpp
--- a/Engine/Plugins/Editor/PluginBrowser/Source/PluginBrowser/Private/PluginMetadataObject.cpp
+++ b/Engine/Plugins/Editor/PluginBrowser/Source/PluginBrowser/Private/PluginMetadataObject.cpp
@@ -60,6 +60,8 @@ void UPluginMetadataObject::PopulateFromPlugin(TSharedPtr<IPlugin> InPlugin)
 		FPluginReferenceMetadata& PluginRef = Plugins.AddDefaulted_GetRef();
 		PluginRef.PopulateFromDescriptor(PluginRefDesc);
 	}
+
+	DisallowedPlugins = InDescriptor.DisallowedPlugins;
 }
 
 void UPluginMetadataObject::CopyIntoDescriptor(FPluginDescriptor& OutDescriptor) const
@@ -102,7 +104,8 @@ void UPluginMetadataObject::CopyIntoDescriptor(FPluginDescriptor& OutDescriptor)
 	}
 
 	OutDescriptor.Plugins = MoveTemp(NewPlugins);
-	
+	OutDescriptor.DisallowedPlugins = DisallowedPlugins;
+
 	// Apply any edits done by an extension
 	for (const TSharedPtr<FPluginEditorExtension>& Extension : Extensions)
 	{
@@ -132,6 +135,16 @@ TArray<FString> UPluginMetadataObject::GetAvailablePluginDependencies() const
 			continue;
 		}
 
+		bool bDependencyDisallowed = false;
+		for (const FString& DisallowedPlugin : DisallowedPlugins)
+		{
+			if (Plugin->GetName() == DisallowedPlugin)
+			{
+				bDependencyDisallowed = true;
+				break;
+			}
+		}
+
 		bool bDependsOnMe = false;
 		for (const FPluginReferenceDescriptor& Dependency : Plugin->GetDescriptor().Plugins)
 		{
@@ -142,7 +155,7 @@ TArray<FString> UPluginMetadataObject::GetAvailablePluginDependencies() const
 			}
 		}
 
-		if (!bDependsOnMe)
+		if (!bDependsOnMe && !bDependencyDisallowed)
 		{
 			Result.Add(Plugin->GetName());
 		}
@@ -151,6 +164,25 @@ TArray<FString> UPluginMetadataObject::GetAvailablePluginDependencies() const
 	return Result;
 }
 
+TArray<FString> UPluginMetadataObject::GetDisallowedPluginsOptions() const
+{
+	TArray<TSharedRef<IPlugin>> AllPlugins = IPluginManager::Get().GetDiscoveredPlugins();
+	TArray<FString> Result;
+	Result.Reserve(AllPlugins.Num());
+
+	for (const TSharedRef<IPlugin>& Plugin : AllPlugins)
+	{
+		if (Plugin == SourcePlugin)
+		{
+			continue;
+		}
+
+		Result.Add(Plugin->GetName());
+	}
+
+	return Result;
+}
+
 /////////////////////////////////////////////////////////////////////////////////////////////////////////
 
 TSharedRef<IDetailCustomization> FPluginMetadataCustomization::MakeInstance()
diff --git a/Engine/Plugins/Editor/PluginBrowser/Source/PluginBrowser/Private/PluginMetadataObject.h b/Engine/Plugins/Editor/PluginBrowser/Source/PluginBrowser/Private/PluginMetadataObject.h
--- a/Engine/Plugins/Editor/PluginBrowser/Source/PluginBrowser/Private/PluginMetadataObject.h
+++ b/Engine/Plugins/Editor/PluginBrowser/Source/PluginBrowser/Private/PluginMetadataObject.h
@@ -132,6 +132,10 @@ public:
 	UPROPERTY(EditAnywhere, Category = Dependencies, meta=(TitleProperty=Name))
 	TArray<FPluginReferenceMetadata> Plugins;
 
+	/** Plugins that cannot be used by this plugin */
+	UPROPERTY(EditAnywhere, Category = Dependencies, meta=(DisplayName="Disallowed", GetOptions= GetDisallowedPluginsOptions))
+	TArray<FString> DisallowedPlugins;
+
 	/** Plugin this proxy object was constructed from */
 	TWeakPtr<IPlugin> SourcePlugin;
 
@@ -150,6 +154,9 @@ public:
 
 	UFUNCTION()
 	TArray<FString> GetAvailablePluginDependencies() const;
+
+	UFUNCTION()
+	TArray<FString> GetDisallowedPluginsOptions() const;
 };
 
 /**
diff --git a/Engine/Source/Programs/UnrealBuildTool/Configuration/UEBuildTarget.cs b/Engine/Source/Programs/UnrealBuildTool/Configuration/UEBuildTarget.cs
--- a/Engine/Source/Programs/UnrealBuildTool/Configuration/UEBuildTarget.cs
+++ b/Engine/Source/Programs/UnrealBuildTool/Configuration/UEBuildTarget.cs
@@ -2273,10 +2273,19 @@ namespace UnrealBuildTool
 			// Check that each plugin does not have a dependency on any sealed plugins
 			foreach (UEBuildPlugin Plugin in BuildPlugins)
 			{
-				UEBuildPlugin? DependencyPlugin = Plugin.Dependencies?.FirstOrDefault(X => X.Descriptor.bIsSealed);
-				if (DependencyPlugin != null)
+				if (Plugin.Dependencies != null)
 				{
-					throw new BuildException("Plugin '{0}' cannot depend on plugin '{1}' because it is sealed.", Plugin.Name, DependencyPlugin.Name);
+					foreach (UEBuildPlugin DependencyPlugin in Plugin.Dependencies)
+					{
+						if (DependencyPlugin.Descriptor.bIsSealed)
+						{
+							throw new BuildException("Plugin '{0}' cannot depend on plugin '{1}' because it is sealed.", Plugin.Name, DependencyPlugin.Name);
+						}
+						else if (Plugin.Descriptor.DisallowedPlugins != null && Plugin.Descriptor.DisallowedPlugins.Contains(DependencyPlugin.Name))
+						{
+							throw new BuildException("Plugin '{0}' cannot depend on plugin '{1}' because it is disallowed.", Plugin.Name, DependencyPlugin.Name);
+						}
+					}
 				}
 			}
 
diff --git a/Engine/Source/Programs/UnrealBuildTool/System/PluginDescriptor.cs b/Engine/Source/Programs/UnrealBuildTool/System/PluginDescriptor.cs
--- a/Engine/Source/Programs/UnrealBuildTool/System/PluginDescriptor.cs
+++ b/Engine/Source/Programs/UnrealBuildTool/System/PluginDescriptor.cs
@@ -216,6 +216,11 @@ namespace UnrealBuildTool
 		/// </summary>
 		public List<PluginReferenceDescriptor>? Plugins;
 
+		/// <summary>
+		/// Plugins that this plugin should never depend on
+		/// </summary>
+		public String[]? DisallowedPlugins;
+
 		/// <summary>
 		/// Private constructor. This object should not be created directly; read it from disk using FromFile() instead.
 		/// </summary>
@@ -346,9 +351,10 @@ namespace UnrealBuildTool
 			{
 				Plugins = Array.ConvertAll(PluginsArray, x => PluginReferenceDescriptor.FromJsonObject(x)).ToList();
 			}
+
+			RawObject.TryGetStringArrayField("DisallowedPlugins", out DisallowedPlugins);
 		}
 
-		/// <summary>
 		/// Creates a plugin descriptor from a file on disk
 		/// </summary>
 		/// <param name="FileName">The filename to read</param>
@@ -493,6 +499,11 @@ namespace UnrealBuildTool
 			{
 				PluginReferenceDescriptor.WriteArray(Writer, "Plugins", Plugins.ToArray());
 			}
+
+			if (DisallowedPlugins != null && DisallowedPlugins.Length > 0)
+			{
+				Writer.WriteStringArrayField("DisallowedPlugins", DisallowedPlugins);
+			}
 		}
 
 		/// <summary>
diff --git a/Engine/Source/Runtime/Projects/Private/PluginDescriptor.cpp b/Engine/Source/Runtime/Projects/Private/PluginDescriptor.cpp
--- a/Engine/Source/Runtime/Projects/Private/PluginDescriptor.cpp
+++ b/Engine/Source/Runtime/Projects/Private/PluginDescriptor.cpp
@@ -267,6 +267,8 @@ bool FPluginDescriptor::Read(const FJsonObject& Object, FText* OutFailReason /*=
 		return false;
 	}
 
+	Object.TryGetStringArrayField(TEXT("DisallowedPlugins"), DisallowedPlugins);
+
 	return true;
 }
 
@@ -475,6 +477,20 @@ void FPluginDescriptor::UpdateJson(FJsonObject& JsonObject) const
 
 	FPluginReferenceDescriptor::UpdateArray(JsonObject, TEXT("Plugins"), Plugins);
 
+	if (DisallowedPlugins.Num() > 0)
+	{
+		TArray<TSharedPtr<FJsonValue>> DisallowedPluginsValues;
+		for (const FString& DisallowedPlugin : DisallowedPlugins)
+		{
+			DisallowedPluginsValues.Add(MakeShareable(new FJsonValueString(DisallowedPlugin)));
+		}
+		JsonObject.SetArrayField(TEXT("DisallowedPlugins"), DisallowedPluginsValues);
+	}
+	else
+	{
+		JsonObject.RemoveField(TEXT("DisallowedPlugins"));
+	}
+
 #if WITH_EDITOR
 	for (const auto& KVP : AdditionalFieldsToWrite)
 	{
diff --git a/Engine/Source/Runtime/Projects/Private/PluginManager.cpp b/Engine/Source/Runtime/Projects/Private/PluginManager.cpp
--- a/Engine/Source/Runtime/Projects/Private/PluginManager.cpp
+++ b/Engine/Source/Runtime/Projects/Private/PluginManager.cpp
@@ -483,6 +483,36 @@ const FString& EnumToString(const EPluginType& InPluginType)
 	}
 }
 
+namespace
+{
+	enum class EConfigurePluginResultCode
+	{
+		PluginMissing,
+		PluginSealed,
+		DisallowedDependency,
+
+		Unset,
+	};
+}
+
+struct FPluginManager::FConfigurePluginResultInfo
+{
+	EConfigurePluginResultCode Code;
+	const FPluginReferenceDescriptor* PluginReference;
+
+	FConfigurePluginResultInfo()
+		: Code(EConfigurePluginResultCode::Unset)
+		, PluginReference(nullptr)
+	{
+	}
+
+	FConfigurePluginResultInfo(EConfigurePluginResultCode InCode, const FPluginReferenceDescriptor* InPluginReference)
+		: Code(InCode)
+		, PluginReference(InPluginReference)
+	{
+	}
+};
+
 FPluginManager::FPluginManager()
 {
 	SCOPED_BOOT_TIMING("DiscoverAllPlugins");
@@ -1866,20 +1896,19 @@ bool FPluginManager::ConfigureEnabledPlugins()
 
 bool FPluginManager::RequiresTempTargetForCodePlugin(const FProjectDescriptor* ProjectDescriptor, const FString& Platform, EBuildConfiguration Configuration, EBuildTargetType TargetType, FText& OutReason)
 {
-	const FPluginReferenceDescriptor* MissingPlugin = nullptr;
-	const FPluginReferenceDescriptor* SealedPlugin = nullptr;
+	FConfigurePluginResultInfo ResultInfo;
 
 	TSet<FString> ProjectCodePlugins;
-	if (!GetCodePluginsForProject(ProjectDescriptor, Platform, Configuration, TargetType, AllPlugins, ProjectCodePlugins, MissingPlugin, SealedPlugin))
+	if (!GetCodePluginsForProject(ProjectDescriptor, Platform, Configuration, TargetType, AllPlugins, ProjectCodePlugins, ResultInfo))
 	{
-		OutReason = FText::Format(LOCTEXT("TempTarget_MissingPluginForTarget", "{0} plugin is referenced by target but not found"), FText::FromString(MissingPlugin->Name));
+		OutReason = FText::Format(LOCTEXT("TempTarget_MissingPluginForTarget", "{0} plugin is referenced by target but not found"), FText::FromString(ResultInfo.PluginReference->Name));
 		return true;
 	}
 
 	TSet<FString> DefaultCodePlugins;
-	if (!GetCodePluginsForProject(nullptr, Platform, Configuration, TargetType, AllPlugins, DefaultCodePlugins, MissingPlugin, SealedPlugin))
+	if (!GetCodePluginsForProject(nullptr, Platform, Configuration, TargetType, AllPlugins, DefaultCodePlugins, ResultInfo))
 	{
-		OutReason = FText::Format(LOCTEXT("TempTarget_MissingPluginForDefaultTarget", "{0} plugin is referenced by the default target but not found"), FText::FromString(MissingPlugin->Name));
+		OutReason = FText::Format(LOCTEXT("TempTarget_MissingPluginForDefaultTarget", "{0} plugin is referenced by the default target but not found"), FText::FromString(ResultInfo.PluginReference->Name));
 		return true;
 	}
 
@@ -1904,7 +1933,7 @@ bool FPluginManager::RequiresTempTargetForCodePlugin(const FProjectDescriptor* P
 	return false;
 }
 
-bool FPluginManager::GetCodePluginsForProject(const FProjectDescriptor* ProjectDescriptor, const FString& Platform, EBuildConfiguration Configuration, EBuildTargetType TargetType, FDiscoveredPluginMap& AllPlugins, TSet<FString>& CodePluginNames, const FPluginReferenceDescriptor*& OutMissingPlugin, const FPluginReferenceDescriptor*& OutSealedPlugin)
+bool FPluginManager::GetCodePluginsForProject(const FProjectDescriptor* ProjectDescriptor, const FString& Platform, EBuildConfiguration Configuration, EBuildTargetType TargetType, FDiscoveredPluginMap& AllPlugins, TSet<FString>& CodePluginNames, FConfigurePluginResultInfo& OutResultInfo)
 {
 	// Can only check the current project at the moment, since we won't have enumerated them otherwise
 	check(ProjectDescriptor == nullptr || ProjectDescriptor == IProjectManager::Get().GetCurrentProject());
@@ -1930,7 +1959,7 @@ bool FPluginManager::GetCodePluginsForProject(const FProjectDescriptor* ProjectD
 		{
 			if(!ConfiguredPluginNames.Contains(PluginReference.Name))
 			{
-				if (!ConfigureEnabledPluginForTarget(PluginReference, ProjectDescriptor, FString(), Platform, Configuration, TargetType, bLoadPluginsForTargetPlatforms, AllPlugins, EnabledPlugins, OutMissingPlugin, OutSealedPlugin))
+				if (!ConfigureEnabledPluginForTarget(PluginReference, ProjectDescriptor, FString(), Platform, Configuration, TargetType, bLoadPluginsForTargetPlatforms, AllPlugins, EnabledPlugins, OutResultInfo))
 				{
 					return false;
 				}
@@ -1944,7 +1973,7 @@ bool FPluginManager::GetCodePluginsForProject(const FProjectDescriptor* ProjectD
 	{
 		if (DiscoveredPluginMapUtils::ResolvePluginFromMapVal(PluginPair.Value)->IsEnabledByDefault(bAllowEnginePluginsEnabledByDefault) && !ConfiguredPluginNames.Contains(PluginPair.Key))
 		{
-			if (!ConfigureEnabledPluginForTarget(FPluginReferenceDescriptor(PluginPair.Key, true), ProjectDescriptor, FString(), Platform, Configuration, TargetType, bLoadPluginsForTargetPlatforms, AllPlugins, EnabledPlugins, OutMissingPlugin, OutSealedPlugin))
+			if (!ConfigureEnabledPluginForTarget(FPluginReferenceDescriptor(PluginPair.Key, true), ProjectDescriptor, FString(), Platform, Configuration, TargetType, bLoadPluginsForTargetPlatforms, AllPlugins, EnabledPlugins, OutResultInfo))
 			{
 				return false;
 			}
@@ -1974,60 +2003,89 @@ bool FPluginManager::ConfigureEnabledPluginForCurrentTarget(const FPluginReferen
 {
 	SCOPED_BOOT_TIMING("ConfigureEnabledPluginForCurrentTarget");
 
-	const FPluginReferenceDescriptor* MissingPlugin = nullptr;
-	const FPluginReferenceDescriptor* SealedPlugin = nullptr;
-	if (!ConfigureEnabledPluginForTarget(FirstReference, IProjectManager::Get().GetCurrentProject(), UE_APP_NAME, FPlatformMisc::GetUBTPlatform(), FApp::GetBuildConfiguration(), FApp::GetBuildTargetType(), (bool)LOAD_PLUGINS_FOR_TARGET_PLATFORMS, AllPlugins, EnabledPlugins, MissingPlugin, SealedPlugin))
+	FConfigurePluginResultInfo ResultInfo;
+
+	if (!ConfigureEnabledPluginForTarget(FirstReference, IProjectManager::Get().GetCurrentProject(), UE_APP_NAME, FPlatformMisc::GetUBTPlatform(), FApp::GetBuildConfiguration(), FApp::GetBuildTargetType(), (bool)LOAD_PLUGINS_FOR_TARGET_PLATFORMS, AllPlugins, EnabledPlugins, ResultInfo))
 	{
 		// If we're in unattended mode, don't open any windows and fatal out
 		if (FApp::IsUnattended())
 		{
-			if (MissingPlugin != nullptr)
-			{
-				UE_LOG(LogPluginManager, Fatal, TEXT("This project requires the '%s' plugin. Install it and try again, or remove it from the project's required plugin list."), *MissingPlugin->Name);
-			}
-			else if (SealedPlugin != nullptr)
+			switch (ResultInfo.Code)
 			{
-				UE_LOG(LogPluginManager, Fatal, TEXT("This project requires the '%s' plugin, which has an illegal dependency on the sealed plugin '%s'."), *FirstReference.Name, *SealedPlugin->Name);
+			case EConfigurePluginResultCode::PluginMissing:
+				UE_LOG(LogPluginManager, Fatal, TEXT("This project requires the '%s' plugin. Install it and try again, or remove it from the project's required plugin list."), *ResultInfo.PluginReference->Name);
+				break;
+
+			case EConfigurePluginResultCode::PluginSealed:
+				UE_LOG(LogPluginManager, Fatal, TEXT("This project requires the '%s' plugin, which has an illegal dependency on the sealed plugin '%s'."), *FirstReference.Name, *ResultInfo.PluginReference->Name);
+				break;
+
+			case EConfigurePluginResultCode::DisallowedDependency:
+				UE_LOG(LogPluginManager, Fatal, TEXT("This project requires the '%s' plugin, which has a disallowed dependency on the plugin '%s'."), *FirstReference.Name, *ResultInfo.PluginReference->Name);
+				break;
+
+			default:
+				checkf(0, TEXT("Unimplemented EConfigurePluginResultCode switch case"));
 			}
 			return false;
 		}
 
 #if !IS_MONOLITHIC
-		if (MissingPlugin != nullptr)
+		switch (ResultInfo.Code)
 		{
-			// Try to download it from the marketplace
-			if (MissingPlugin->MarketplaceURL.Len() > 0 && PromptToDownloadPlugin(MissingPlugin->Name, MissingPlugin->MarketplaceURL))
+		case EConfigurePluginResultCode::PluginMissing:
 			{
-				UE_LOG(LogPluginManager, Display, TEXT("Downloading '%s' plugin from marketplace (%s)."), *MissingPlugin->Name, *MissingPlugin->MarketplaceURL);
-				return false;
+				// Try to download it from the marketplace
+				if (ResultInfo.PluginReference->MarketplaceURL.Len() > 0 && PromptToDownloadPlugin(ResultInfo.PluginReference->Name, ResultInfo.PluginReference->MarketplaceURL))
+				{
+					UE_LOG(LogPluginManager, Display, TEXT("Downloading '%s' plugin from marketplace (%s)."), *ResultInfo.PluginReference->Name, *ResultInfo.PluginReference->MarketplaceURL);
+					return false;
+				}
+
+				// Prompt to disable it in the project file, if possible
+				if (PromptToDisableMissingPlugin(FirstReference.Name, ResultInfo.PluginReference->Name))
+				{
+					UE_LOG(LogPluginManager, Display, TEXT("Disabled plugin '%s' due to missing plugin '%s', continuing."), *FirstReference.Name, *ResultInfo.PluginReference->Name);
+					return true;
+				}
 			}
+			break;
 
-			// Prompt to disable it in the project file, if possible
-			if (PromptToDisableMissingPlugin(FirstReference.Name, MissingPlugin->Name))
+		case EConfigurePluginResultCode::PluginSealed:
 			{
-				UE_LOG(LogPluginManager, Display, TEXT("Disabled plugin '%s', continuing."), *FirstReference.Name);
-				return true;
+				// Prompt to disable it in the project file, if possible
+				if (PromptToDisableSealedPlugin(FirstReference.Name, ResultInfo.PluginReference->Name))
+				{
+					UE_LOG(LogPluginManager, Display, TEXT("Disabled plugin '%s' due to sealed dependency on plugin '%s', continuing."), *FirstReference.Name, *ResultInfo.PluginReference->Name);
+					return true;
+				}
 			}
-		}
-		else if (SealedPlugin != nullptr)
-		{
-			// Prompt to disable it in the project file, if possible
-			if (PromptToDisableSealedPlugin(FirstReference.Name, SealedPlugin->Name))
+			break;
+
+		case EConfigurePluginResultCode::DisallowedDependency:
 			{
-				UE_LOG(LogPluginManager, Display, TEXT("Disabled plugin '%s', continuing."), *FirstReference.Name);
-				return true;
+				// Prompt to disable it in the project file, if possible
+				if (PromptToDisableDisalowedPlugin(FirstReference.Name, ResultInfo.PluginReference->Name))
+				{
+					UE_LOG(LogPluginManager, Display, TEXT("Disabled plugin '%s' due to disallowed dependency on plugin '%s', continuing."), *FirstReference.Name, *ResultInfo.PluginReference->Name);
+					return true;
+				}
 			}
+			break;
+
+		default:
+			checkf(0, TEXT("Unimplemented EConfigurePluginResultCode switch case"));
 		}
 #endif
 
 		// Unable to continue
-		UE_LOG(LogPluginManager, Error, TEXT("Unable to load plugin '%s'. Aborting."), MissingPlugin ? *MissingPlugin->Name : *SealedPlugin->Name);
+		UE_LOG(LogPluginManager, Error, TEXT("Unable to load plugin '%s'. Aborting."), *ResultInfo.PluginReference->Name);
 		return false;
 	}
 	return true;
 }
 
-bool FPluginManager::ConfigureEnabledPluginForTarget(const FPluginReferenceDescriptor& FirstReference, const FProjectDescriptor* ProjectDescriptor, const FString& TargetName, const FString& Platform, EBuildConfiguration Configuration, EBuildTargetType TargetType, bool bLoadPluginsForTargetPlatforms, FDiscoveredPluginMap& AllPlugins, TMap<FString, FPlugin*>& EnabledPlugins, const FPluginReferenceDescriptor*& OutMissingPlugin, const FPluginReferenceDescriptor*& OutSealedPlugin)
+bool FPluginManager::ConfigureEnabledPluginForTarget(const FPluginReferenceDescriptor& FirstReference, const FProjectDescriptor* ProjectDescriptor, const FString& TargetName, const FString& Platform, EBuildConfiguration Configuration, EBuildTargetType TargetType, bool bLoadPluginsForTargetPlatforms, FDiscoveredPluginMap& AllPlugins, TMap<FString, FPlugin*>& EnabledPlugins, FConfigurePluginResultInfo& OutResultInfo)
 {
 	if (!EnabledPlugins.Contains(FirstReference.Name))
 	{
@@ -2106,7 +2164,7 @@ bool FPluginManager::ConfigureEnabledPluginForTarget(const FPluginReferenceDescr
 				}
 
 				// Add it to the missing list
-				OutMissingPlugin = &Reference;
+				OutResultInfo = FConfigurePluginResultInfo(EConfigurePluginResultCode::PluginMissing, &Reference);
 				return false;
 			}
 
@@ -2169,7 +2227,13 @@ bool FPluginManager::ConfigureEnabledPluginForTarget(const FPluginReferenceDescr
 				const TSharedRef<FPlugin>* NextPluginPtr = DiscoveredPluginMapUtils::FindPluginInMap(AllPlugins, NextReference.Name);
 				if (NextPluginPtr != nullptr && NextPluginPtr->Get().Descriptor.bIsSealed)
 				{
-					OutSealedPlugin = &NextReference;
+					OutResultInfo = FConfigurePluginResultInfo(EConfigurePluginResultCode::PluginSealed, &NextReference);
+					return false;
+				}
+
+				if (NextPluginPtr != nullptr && Plugin.Descriptor.DisallowedPlugins.Contains(NextPluginPtr->Get().Name))
+				{
+					OutResultInfo = FConfigurePluginResultInfo(EConfigurePluginResultCode::DisallowedDependency, &NextReference);
 					return false;
 				}
 
@@ -2227,6 +2291,14 @@ bool FPluginManager::PromptToDisableSealedPlugin(const FString& PluginName, cons
 	return PromptToDisablePlugin(Caption, Message, PluginName);
 }
 
+bool FPluginManager::PromptToDisableDisalowedPlugin(const FString& PluginName, const FString& DisallowedPluginName)
+{
+	FText Message = FText::Format(LOCTEXT("DisablePluginMessage_DisallowedDependency", "This project requires the '{0}' plugin, which has a disallowed dependency on the '{1}' plugin.\n\nWould you like to disable it?\n\nIf you do, you will no longer be able to open any assets created with it. If not, the application will close."), FText::FromString(PluginName), FText::FromString(DisallowedPluginName));
+	FText Caption(LOCTEXT("DisablePluginCaption_IllegalDependency", "Illegal Dependency"));
+
+	return PromptToDisablePlugin(Caption, Message, PluginName);
+}
+
 bool FPluginManager::PromptToDisableMissingPlugin(const FString& PluginName, const FString& MissingPluginName)
 {
 	FText Message;
diff --git a/Engine/Source/Runtime/Projects/Private/PluginManager.h b/Engine/Source/Runtime/Projects/Private/PluginManager.h
--- a/Engine/Source/Runtime/Projects/Private/PluginManager.h
+++ b/Engine/Source/Runtime/Projects/Private/PluginManager.h
@@ -195,6 +195,8 @@ public:
 private:
 	using FDiscoveredPluginMap = TMap<FString, TArray<TSharedRef<FPlugin>>>;
 
+	struct FConfigurePluginResultInfo;
+
 	/** Searches for all plugins on disk and builds up the array of plugin objects.  Doesn't load any plugins. 
 	    This is called when the plugin manager singleton is first accessed. */
 	void DiscoverAllPlugins();
@@ -215,7 +217,7 @@ private:
 	static void FindPluginManifestsInDirectory(const FString& PluginManifestDirectory, TArray<FString>& FileNames);
 
 	/** Gets all the code plugins that are enabled for a content only project */
-	static bool GetCodePluginsForProject(const FProjectDescriptor* ProjectDescriptor, const FString& Platform, EBuildConfiguration Configuration, EBuildTargetType TargetType, FDiscoveredPluginMap& AllPlugins, TSet<FString>& CodePluginNames, const FPluginReferenceDescriptor*& OutMissingPlugin, const FPluginReferenceDescriptor*& OutSealedPlugin);
+	static bool GetCodePluginsForProject(const FProjectDescriptor* ProjectDescriptor, const FString& Platform, EBuildConfiguration Configuration, EBuildTargetType TargetType, FDiscoveredPluginMap& AllPlugins, TSet<FString>& CodePluginNames, FConfigurePluginResultInfo& OutResultInfo);
 
 	/** Sets the bPluginEnabled flag on all plugins found from DiscoverAllPlugins that are enabled in config */
 	bool ConfigureEnabledPlugins();
@@ -224,7 +226,7 @@ private:
 	bool ConfigureEnabledPluginForCurrentTarget(const FPluginReferenceDescriptor& FirstReference, TMap<FString, FPlugin*>& EnabledPlugins);
 
 	/** Adds a single enabled plugin and all its dependencies. */
-	static bool ConfigureEnabledPluginForTarget(const FPluginReferenceDescriptor& FirstReference, const FProjectDescriptor* ProjectDescriptor, const FString& TargetName, const FString& Platform, EBuildConfiguration Configuration, EBuildTargetType TargetType, bool bLoadPluginsForTargetPlatforms, FDiscoveredPluginMap& AllPlugins, TMap<FString, FPlugin*>& EnabledPlugins, const FPluginReferenceDescriptor*& OutMissingPlugin, const FPluginReferenceDescriptor*& OutSealedPlugin);
+	static bool ConfigureEnabledPluginForTarget(const FPluginReferenceDescriptor& FirstReference, const FProjectDescriptor* ProjectDescriptor, const FString& TargetName, const FString& Platform, EBuildConfiguration Configuration, EBuildTargetType TargetType, bool bLoadPluginsForTargetPlatforms, FDiscoveredPluginMap& AllPlugins, TMap<FString, FPlugin*>& EnabledPlugins, FConfigurePluginResultInfo& OutResultInfo);
 
 	/** Prompts the user to download a missing plugin from the given URL */
 	static bool PromptToDownloadPlugin(const FString& PluginName, const FString& MarketplaceURL);
@@ -232,6 +234,9 @@ private:
 	/** Prompts the user to disable a plugin */
 	static bool PromptToDisableSealedPlugin(const FString& PluginName, const FString& SealedPluginName);
 
+	/** Prompts the user to disable a plugin */
+	static bool PromptToDisableDisalowedPlugin(const FString& PluginName, const FString& DisallowedPluginName);
+
 	/** Prompts the user to disable a plugin */
 	static bool PromptToDisableMissingPlugin(const FString& PluginName, const FString& MissingPluginName);
 
diff --git a/Engine/Source/Runtime/Projects/Public/PluginDescriptor.h b/Engine/Source/Runtime/Projects/Public/PluginDescriptor.h
--- a/Engine/Source/Runtime/Projects/Public/PluginDescriptor.h
+++ b/Engine/Source/Runtime/Projects/Public/PluginDescriptor.h
@@ -146,6 +146,9 @@ struct PROJECTS_API FPluginDescriptor
 	/** Plugins used by this plugin */
 	TArray<FPluginReferenceDescriptor> Plugins;
 
+	/** Plugins that cannot be used by this plugin */
+	TArray<FString> DisallowedPlugins;
+
 #if WITH_EDITOR
 	/** Cached json for custom data */
 	mutable TSharedPtr<FJsonObject> CachedJson;
