Only support UE_ADDITIONAL_PLUGIN_PATHS environment variable in editor builds.

diff --git a/Engine/Source/Runtime/Projects/Private/PluginManager.cpp b/Engine/Source/Runtime/Projects/Private/PluginManager.cpp
--- a/Engine/Source/Runtime/Projects/Private/PluginManager.cpp
+++ b/Engine/Source/Runtime/Projects/Private/PluginManager.cpp
@@ -67,7 +67,11 @@ TArray<FString> GetPluginPathsByEnv( const TCHAR* EnvVariable )
  */
 TArray<FString> GetAdditionalExternalPluginsByEnvVar()
 {
+#if WITH_EDITOR
 	return GetPluginPathsByEnv(TEXT("UE_ADDITIONAL_PLUGIN_PATHS"));
+#else
+	return {};
+#endif
 }
 
 }
