New delegate event IPluginManager::OnNewPluginContentMounted
- this allows us to run code after a plugin is mounted that specifically depends on the plugin's mount point existing

diff --git a/Engine/Source/Runtime/Projects/Private/PluginManager.cpp b/Engine/Source/Runtime/Projects/Private/PluginManager.cpp
--- a/Engine/Source/Runtime/Projects/Private/PluginManager.cpp
+++ b/Engine/Source/Runtime/Projects/Private/PluginManager.cpp
@@ -1886,7 +1886,12 @@ bool FPluginManager::ConfigureEnabledPlugins()
 						FString ContentDir = Plugin.GetContentDir();
 						RegisterMountPointDelegate.Execute(Plugin.GetMountedAssetPath(), ContentDir);
 					}
-				}				
+
+					if (NewPluginContentMountedEvent.IsBound())
+					{
+						NewPluginContentMountedEvent.Broadcast(Plugin);
+					}
+				}
 			}
 		}
 
@@ -2777,6 +2782,11 @@ IPluginManager::FNewPluginMountedEvent& FPluginManager::OnNewPluginMounted()
 	return NewPluginMountedEvent;
 }
 
+IPluginManager::FNewPluginMountedEvent& FPluginManager::OnNewPluginContentMounted()
+{
+	return NewPluginContentMountedEvent;
+}
+
 IPluginManager::FNewPluginMountedEvent& FPluginManager::OnPluginEdited()
 {
 	return PluginEditedEvent;
@@ -2873,7 +2883,8 @@ void FPluginManager::MountPluginFromExternalSource(const TSharedRef<FPlugin>& Pl
 	Plugin->bEnabled = true;
 
 	// Mount the plugin content directory
-	if ((Plugin->CanContainContent() || Plugin->CanContainVerse()) && ensure(RegisterMountPointDelegate.IsBound()))
+	const bool bHasContentOrVerse = (Plugin->CanContainContent() || Plugin->CanContainVerse()) && ensure(RegisterMountPointDelegate.IsBound());
+	if (bHasContentOrVerse)
 	{
 		if (NewPluginMountedEvent.IsBound())
 		{
@@ -2926,6 +2937,12 @@ void FPluginManager::MountPluginFromExternalSource(const TSharedRef<FPlugin>& Pl
 
 	Plugin->SetIsMounted(true);
 
+	// Notify listeners that the plugin is completely mounted now
+	if (bHasContentOrVerse && NewPluginContentMountedEvent.IsBound())
+	{
+		NewPluginContentMountedEvent.Broadcast(*Plugin);
+	}
+
 	if (GWarn)
 	{
 		GWarn->EndSlowTask();
diff --git a/Engine/Source/Runtime/Projects/Private/PluginManager.h b/Engine/Source/Runtime/Projects/Private/PluginManager.h
--- a/Engine/Source/Runtime/Projects/Private/PluginManager.h
+++ b/Engine/Source/Runtime/Projects/Private/PluginManager.h
@@ -170,6 +170,7 @@ public:
 	virtual TArray<TSharedRef<IPlugin>> GetPluginsWithPakFile() const override;
 	virtual FNewPluginMountedEvent& OnNewPluginCreated() override;
 	virtual FNewPluginMountedEvent& OnNewPluginMounted() override;
+	virtual FNewPluginMountedEvent& OnNewPluginContentMounted() override;
 	virtual FNewPluginMountedEvent& OnPluginEdited() override;
 	virtual FNewPluginMountedEvent& OnPluginUnmounted() override;
 	virtual void MountNewlyCreatedPlugin(const FString& PluginName) override;
@@ -317,6 +318,7 @@ private:
 	/** Callback for notifications that a new plugin was mounted */
 	FNewPluginMountedEvent NewPluginCreatedEvent;
 	FNewPluginMountedEvent NewPluginMountedEvent;
+	FNewPluginMountedEvent NewPluginContentMountedEvent;
 	FNewPluginMountedEvent PluginEditedEvent;
 	FNewPluginMountedEvent PluginUnmountedEvent;
 
diff --git a/Engine/Source/Runtime/Projects/Public/Interfaces/IPluginManager.h b/Engine/Source/Runtime/Projects/Public/Interfaces/IPluginManager.h
--- a/Engine/Source/Runtime/Projects/Public/Interfaces/IPluginManager.h
+++ b/Engine/Source/Runtime/Projects/Public/Interfaces/IPluginManager.h
@@ -402,6 +402,11 @@ public:
 	 */
 	virtual FNewPluginMountedEvent& OnNewPluginMounted() = 0;
 
+	/**
+	 * Event for being notified that a new plugin and its content have been mounted
+	 */
+	virtual FNewPluginMountedEvent& OnNewPluginContentMounted() = 0;
+
 	/**
 	 * Event for being notified that a plugin has been edited
 	 */
