PluginManager: Remove extraneous newline from logmessage.
#rnx

diff --git a/Engine/Source/Runtime/Projects/Private/PluginManager.cpp b/Engine/Source/Runtime/Projects/Private/PluginManager.cpp
--- a/Engine/Source/Runtime/Projects/Private/PluginManager.cpp
+++ b/Engine/Source/Runtime/Projects/Private/PluginManager.cpp
@@ -937,7 +937,7 @@ void FPluginManager::ReadAllPlugins(FDiscoveredPluginMap& Plugins, const TSet<FS
 				if (!bFound)
 				{
 					Parent->Descriptor.Modules.Add(ChildModule);
-					UE_LOG(LogPluginManager, Log, TEXT("Adding extension module %s from %s to %s\n"), *ChildModule.Name.ToString(), *Child->Name, *Parent->Name);
+					UE_LOG(LogPluginManager, Log, TEXT("Adding extension module %s from %s to %s"), *ChildModule.Name.ToString(), *Child->Name, *Parent->Name);
 				}
 
 			}
