Disallowed plugin entries can now contain a developer comment. Changed data type from a string to a struct. Backwards compatibility retained

diff --git a/Engine/Plugins/Editor/PluginBrowser/Source/PluginBrowser/Private/PluginMetadataObject.cpp b/Engine/Plugins/Editor/PluginBrowser/Source/PluginBrowser/Private/PluginMetadataObject.cpp
--- a/Engine/Plugins/Editor/PluginBrowser/Source/PluginBrowser/Private/PluginMetadataObject.cpp
+++ b/Engine/Plugins/Editor/PluginBrowser/Source/PluginBrowser/Private/PluginMetadataObject.cpp
@@ -5,6 +5,7 @@
 #include "Widgets/SBoxPanel.h"
 #include "Widgets/Text/STextBlock.h"
 #include "Interfaces/IPluginManager.h"
+#include "PluginDisallowedDescriptor.h"
 #include "PluginReferenceDescriptor.h"
 #include "DetailLayoutBuilder.h"
 #include "DetailWidgetRow.h"
@@ -27,6 +28,18 @@ void FPluginReferenceMetadata::CopyIntoDescriptor(FPluginReferenceDescriptor& Ou
 	OutDescriptor.bOptional = bOptional;
 }
 
+void FPluginDisallowedMetadata::PopulateFromDescriptor(const FPluginDisallowedDescriptor& InDescriptor)
+{
+	Name = InDescriptor.Name;
+	Comment = InDescriptor.Comment;
+}
+
+void FPluginDisallowedMetadata::CopyIntoDescriptor(FPluginDisallowedDescriptor& OutDescriptor) const
+{
+	OutDescriptor.Name = Name;
+	OutDescriptor.Comment = Comment;
+}
+
 UPluginMetadataObject::UPluginMetadataObject(const FObjectInitializer& ObjectInitializer)
 {
 }
@@ -61,7 +74,12 @@ void UPluginMetadataObject::PopulateFromPlugin(TSharedPtr<IPlugin> InPlugin)
 		PluginRef.PopulateFromDescriptor(PluginRefDesc);
 	}
 
-	DisallowedPlugins = InDescriptor.DisallowedPlugins;
+	DisallowedPlugins.Reset(InDescriptor.DisallowedPlugins.Num());
+	for (const FPluginDisallowedDescriptor& PluginRefDesc : InDescriptor.DisallowedPlugins)
+	{
+		FPluginDisallowedMetadata& PluginRef = DisallowedPlugins.AddDefaulted_GetRef();
+		PluginRef.PopulateFromDescriptor(PluginRefDesc);
+	}
 }
 
 void UPluginMetadataObject::CopyIntoDescriptor(FPluginDescriptor& OutDescriptor) const
@@ -104,7 +122,24 @@ void UPluginMetadataObject::CopyIntoDescriptor(FPluginDescriptor& OutDescriptor)
 	}
 
 	OutDescriptor.Plugins = MoveTemp(NewPlugins);
-	OutDescriptor.DisallowedPlugins = DisallowedPlugins;
+
+	{
+		TArray<FPluginDisallowedDescriptor> NewDisallowedPlugins;
+		NewDisallowedPlugins.Reserve(DisallowedPlugins.Num());
+
+		for (const FPluginDisallowedMetadata& DisallowedRefMetadata : DisallowedPlugins)
+		{
+			if (DisallowedRefMetadata.Name.IsEmpty())
+			{
+				continue;
+			}
+
+			FPluginDisallowedDescriptor& NewDisallowedRefDesc = NewDisallowedPlugins.AddDefaulted_GetRef();
+			DisallowedRefMetadata.CopyIntoDescriptor(NewDisallowedRefDesc);
+		}
+
+		OutDescriptor.DisallowedPlugins = MoveTemp(NewDisallowedPlugins);
+	}
 
 	// Apply any edits done by an extension
 	for (const TSharedPtr<FPluginEditorExtension>& Extension : Extensions)
@@ -136,9 +171,9 @@ TArray<FString> UPluginMetadataObject::GetAvailablePluginDependencies() const
 		}
 
 		bool bDependencyDisallowed = false;
-		for (const FString& DisallowedPlugin : DisallowedPlugins)
+		for (const FPluginDisallowedDescriptor& DisallowedPlugin : Plugin->GetDescriptor().DisallowedPlugins)
 		{
-			if (Plugin->GetName() == DisallowedPlugin)
+			if (Plugin->GetName() == DisallowedPlugin.Name)
 			{
 				bDependencyDisallowed = true;
 				break;
diff --git a/Engine/Plugins/Editor/PluginBrowser/Source/PluginBrowser/Private/PluginMetadataObject.h b/Engine/Plugins/Editor/PluginBrowser/Source/PluginBrowser/Private/PluginMetadataObject.h
--- a/Engine/Plugins/Editor/PluginBrowser/Source/PluginBrowser/Private/PluginMetadataObject.h
+++ b/Engine/Plugins/Editor/PluginBrowser/Source/PluginBrowser/Private/PluginMetadataObject.h
@@ -7,6 +7,7 @@
 
 class IDetailLayoutBuilder;
 struct FPluginDescriptor;
+struct FPluginDisallowedDescriptor;
 struct FPluginReferenceDescriptor;
 class IPlugin;
 struct FPluginEditorExtension;
@@ -43,6 +44,30 @@ public:
 	void CopyIntoDescriptor(FPluginReferenceDescriptor& OutDescriptor) const;
 };
 
+USTRUCT()
+struct FPluginDisallowedMetadata
+{
+public:
+	GENERATED_BODY()
+
+	/** Name of the dependency plugin */
+	UPROPERTY(EditAnywhere, Category = Details, meta = (GetOptions = GetDisallowedPluginsOptions))
+	FString Name;
+
+	UPROPERTY(EditAnywhere, Category = Details)
+	FString Comment;
+
+	/**
+	 * Populate the fields of this object from an existing descriptor.
+	 */
+	void PopulateFromDescriptor(const FPluginDisallowedDescriptor& InDescriptor);
+
+	/**
+	 * Copy the metadata fields into a plugin descriptor.
+	 */
+	void CopyIntoDescriptor(FPluginDisallowedDescriptor& OutDescriptor) const;	
+};
+
 /**
  * We use this object to display plugin properties using details view.
  */
@@ -133,8 +158,8 @@ public:
 	TArray<FPluginReferenceMetadata> Plugins;
 
 	/** Plugins that cannot be used by this plugin */
-	UPROPERTY(EditAnywhere, Category = Dependencies, meta=(DisplayName="Disallowed", GetOptions= GetDisallowedPluginsOptions))
-	TArray<FString> DisallowedPlugins;
+	UPROPERTY(EditAnywhere, Category = Dependencies, meta = (DisplayName = "Disallowed"))
+	TArray<FPluginDisallowedMetadata> DisallowedPlugins;
 
 	/** Plugin this proxy object was constructed from */
 	TWeakPtr<IPlugin> SourcePlugin;
diff --git a/Engine/Source/Runtime/Projects/Private/PluginDescriptor.cpp b/Engine/Source/Runtime/Projects/Private/PluginDescriptor.cpp
--- a/Engine/Source/Runtime/Projects/Private/PluginDescriptor.cpp
+++ b/Engine/Source/Runtime/Projects/Private/PluginDescriptor.cpp
@@ -267,7 +267,21 @@ bool FPluginDescriptor::Read(const FJsonObject& Object, FText* OutFailReason /*=
 		return false;
 	}
 
-	Object.TryGetStringArrayField(TEXT("DisallowedPlugins"), DisallowedPlugins);
+	// Backwards compatibility support
+	TArray<FString> DisallowedPluginNameStrings;
+	if (Object.TryGetStringArrayField(TEXT("DisallowedPlugins"), DisallowedPluginNameStrings))
+	{
+		DisallowedPlugins.Reserve(DisallowedPluginNameStrings.Num());
+		for (int32 Index = 0; Index < DisallowedPluginNameStrings.Num(); ++Index)
+		{
+			FPluginDisallowedDescriptor& PluginDisallowedDescriptor = DisallowedPlugins.AddDefaulted_GetRef();
+			PluginDisallowedDescriptor.Name = DisallowedPluginNameStrings[Index];
+		}
+	}
+	else if (!FPluginDisallowedDescriptor::ReadArray(Object, TEXT("DisallowedPlugins"), DisallowedPlugins, OutFailReason))
+	{
+		return false;
+	}
 
 	return true;
 }
@@ -477,19 +491,7 @@ void FPluginDescriptor::UpdateJson(FJsonObject& JsonObject) const
 
 	FPluginReferenceDescriptor::UpdateArray(JsonObject, TEXT("Plugins"), Plugins);
 
-	if (DisallowedPlugins.Num() > 0)
-	{
-		TArray<TSharedPtr<FJsonValue>> DisallowedPluginsValues;
-		for (const FString& DisallowedPlugin : DisallowedPlugins)
-		{
-			DisallowedPluginsValues.Add(MakeShareable(new FJsonValueString(DisallowedPlugin)));
-		}
-		JsonObject.SetArrayField(TEXT("DisallowedPlugins"), DisallowedPluginsValues);
-	}
-	else
-	{
-		JsonObject.RemoveField(TEXT("DisallowedPlugins"));
-	}
+	FPluginDisallowedDescriptor::UpdateArray(JsonObject, TEXT("DisallowedPlugins"), DisallowedPlugins);
 
 #if WITH_EDITOR
 	for (const auto& KVP : AdditionalFieldsToWrite)
diff --git a/Engine/Source/Runtime/Projects/Private/PluginDisallowedDescriptor.cpp b/Engine/Source/Runtime/Projects/Private/PluginDisallowedDescriptor.cpp
new file mode 100644
--- /dev/null
+++ b/Engine/Source/Runtime/Projects/Private/PluginDisallowedDescriptor.cpp
@@ -0,0 +1,121 @@
+// Copyright Epic Games, Inc. All Rights Reserved.
+
+#include "PluginDisallowedDescriptor.h"
+#include "Misc/FileHelper.h"
+#include "Serialization/JsonReader.h"
+#include "Serialization/JsonSerializer.h"
+#include "JsonUtils/JsonObjectArrayUpdater.h"
+#include "ProjectDescriptor.h"
+#include "JsonExtensions.h"
+
+#define LOCTEXT_NAMESPACE "PluginDisallowedDescriptor"
+
+namespace PluginDisallowedDescriptor
+{
+	FString GetPluginRefKey(const FPluginDisallowedDescriptor& PluginRef)
+	{
+		return PluginRef.Name;
+	}
+
+	bool TryGetPluginRefJsonObjectKey(const FJsonObject& JsonObject, FString& OutKey)
+	{
+		return JsonObject.TryGetStringField(TEXT("Name"), OutKey);
+	}
+
+	void UpdatePluginRefJsonObject(const FPluginDisallowedDescriptor& PluginRef, FJsonObject& JsonObject)
+	{
+		PluginRef.UpdateJson(JsonObject);
+	}
+}
+
+bool FPluginDisallowedDescriptor::Read(const FJsonObject& Object, FText* OutFailReason /*= nullptr*/)
+{
+	bool bSuccess = true;
+
+	// Get the name
+	if (!Object.TryGetStringField(TEXT("Name"), Name))
+	{
+		if (OutFailReason)
+		{
+			*OutFailReason = LOCTEXT("FPluginDisallowedWithoutName", "Disallowed plugin must have a 'Name' field");
+		}
+		bSuccess = false;
+	}
+
+#if WITH_EDITOR
+	Object.TryGetStringField(TEXT("Comment"), Comment);
+#endif //if WITH_EDITOR
+
+	return bSuccess;
+}
+
+void FPluginDisallowedDescriptor::Write(TJsonWriter<>& Writer) const
+{
+	TSharedPtr<FJsonObject> PluginRefJsonObject = MakeShared<FJsonObject>();
+
+	UpdateJson(*PluginRefJsonObject);
+
+	FJsonSerializer::Serialize(PluginRefJsonObject.ToSharedRef(), Writer);
+}
+
+void FPluginDisallowedDescriptor::UpdateJson(FJsonObject& JsonObject) const
+{
+	JsonObject.SetStringField(TEXT("Name"), Name);
+
+#if WITH_EDITOR
+	JsonObject.SetStringField(TEXT("Comment"), Comment);
+#endif //if WITH_EDITOR
+}
+
+bool FPluginDisallowedDescriptor::ReadArray(const FJsonObject& Object, const TCHAR* Name, TArray<FPluginDisallowedDescriptor>& OutPlugins, FText* OutFailReason /*= nullptr*/)
+{
+	const TArray< TSharedPtr<FJsonValue> >* Array;
+
+	if (Object.TryGetArrayField(Name, Array))
+	{
+		for (const TSharedPtr<FJsonValue>& Item : *Array)
+		{
+			const TSharedPtr<FJsonObject>* ItemObject = nullptr;
+
+			if (Item.IsValid() && Item->TryGetObject(ItemObject) && ItemObject && ItemObject->IsValid())
+			{
+				FPluginDisallowedDescriptor DisallowedRef;
+
+				if (!DisallowedRef.Read(*ItemObject->Get(), OutFailReason))
+				{
+					return false;
+				}
+
+				OutPlugins.Add(MoveTemp(DisallowedRef));
+			}
+		}
+	}
+
+	return true;
+}
+
+void FPluginDisallowedDescriptor::WriteArray(TJsonWriter<>& Writer, const TCHAR* ArrayName, const TArray<FPluginDisallowedDescriptor>& Plugins)
+{
+	if (Plugins.Num() > 0)
+	{
+		Writer.WriteArrayStart(ArrayName);
+
+		for (const FPluginDisallowedDescriptor& PluginRef : Plugins)
+		{
+			PluginRef.Write(Writer);
+		}
+
+		Writer.WriteArrayEnd();
+	}
+}
+
+void FPluginDisallowedDescriptor::UpdateArray(FJsonObject& JsonObject, const TCHAR* ArrayName, const TArray<FPluginDisallowedDescriptor>& Plugins)
+{
+	typedef FJsonObjectArrayUpdater<FPluginDisallowedDescriptor, FString> FPluginRefJsonArrayUpdater;
+
+	FPluginRefJsonArrayUpdater::Execute(
+		JsonObject, ArrayName, Plugins,
+		FPluginRefJsonArrayUpdater::FGetElementKey::CreateStatic(PluginDisallowedDescriptor::GetPluginRefKey),
+		FPluginRefJsonArrayUpdater::FTryGetJsonObjectKey::CreateStatic(PluginDisallowedDescriptor::TryGetPluginRefJsonObjectKey),
+		FPluginRefJsonArrayUpdater::FUpdateJsonObject::CreateStatic(PluginDisallowedDescriptor::UpdatePluginRefJsonObject));
+}
\ No newline at end of file
diff --git a/Engine/Source/Runtime/Projects/Private/PluginManager.cpp b/Engine/Source/Runtime/Projects/Private/PluginManager.cpp
--- a/Engine/Source/Runtime/Projects/Private/PluginManager.cpp
+++ b/Engine/Source/Runtime/Projects/Private/PluginManager.cpp
@@ -2249,7 +2249,7 @@ bool FPluginManager::ConfigureEnabledPluginForTarget(const FPluginReferenceDescr
 					return false;
 				}
 
-				if (NextPluginPtr != nullptr && Plugin.Descriptor.DisallowedPlugins.Contains(NextPluginPtr->Get().Name))
+				if (NextPluginPtr != nullptr && Plugin.Descriptor.DisallowedPlugins.ContainsByPredicate([&NextPluginPtr](const FPluginDisallowedDescriptor& Other) { return Other.Name == (*NextPluginPtr)->Name; }))
 				{
 					OutResultInfo = FConfigurePluginResultInfo(EConfigurePluginResultCode::DisallowedDependency, &NextReference);
 					return false;
diff --git a/Engine/Source/Runtime/Projects/Public/PluginDescriptor.h b/Engine/Source/Runtime/Projects/Public/PluginDescriptor.h
--- a/Engine/Source/Runtime/Projects/Public/PluginDescriptor.h
+++ b/Engine/Source/Runtime/Projects/Public/PluginDescriptor.h
@@ -10,6 +10,7 @@
 #include "HAL/Platform.h"
 #include "LocalizationDescriptor.h"
 #include "ModuleDescriptor.h"
+#include "PluginDisallowedDescriptor.h"
 #include "PluginReferenceDescriptor.h"
 #include "Serialization/JsonWriter.h"
 #include "Templates/SharedPointer.h"
@@ -147,7 +148,8 @@ struct FPluginDescriptor
 	TArray<FPluginReferenceDescriptor> Plugins;
 
 	/** Plugins that cannot be used by this plugin */
-	TArray<FString> DisallowedPlugins;
+	TArray<FPluginDisallowedDescriptor> DisallowedPlugins;
+
 
 #if WITH_EDITOR
 	/** Cached json for custom data */
diff --git a/Engine/Source/Runtime/Projects/Public/PluginDisallowedDescriptor.h b/Engine/Source/Runtime/Projects/Public/PluginDisallowedDescriptor.h
new file mode 100644
--- /dev/null
+++ b/Engine/Source/Runtime/Projects/Public/PluginDisallowedDescriptor.h
@@ -0,0 +1,38 @@
+// Copyright Epic Games, Inc. All Rights Reserved.
+
+#pragma once
+
+#include "Containers/Array.h"
+#include "Containers/UnrealString.h"
+#include "CoreMinimal.h"
+#include "CoreTypes.h"
+#include "Serialization/JsonWriter.h"
+
+class FJsonObject;
+class FText;
+
+struct FPluginDisallowedDescriptor
+{
+	FString Name;
+
+#if WITH_EDITOR
+	FString Comment;
+#endif
+
+	/** Reads the descriptor from the given JSON object */
+	PROJECTS_API bool Read(const FJsonObject& Object, FText* OutFailReason = nullptr);
+
+	/** Writes a descriptor to JSON */
+	PROJECTS_API void Write(TJsonWriter<>& Writer) const;
+
+	/** Updates the given json object with values in this descriptor */
+	PROJECTS_API void UpdateJson(FJsonObject& JsonObject) const;
+
+	static PROJECTS_API bool ReadArray(const FJsonObject& Object, const TCHAR* Name, TArray<FPluginDisallowedDescriptor>& OutPlugins, FText* OutFailReason = nullptr);
+
+	/** Writes an array of plugin references to JSON */
+	static PROJECTS_API void WriteArray(TJsonWriter<>& Writer, const TCHAR* ArrayName, const TArray<FPluginDisallowedDescriptor>& Plugins);
+
+	/** Updates an array of plugin references in the specified JSON field (indexed by plugin name) */
+	static PROJECTS_API void UpdateArray(FJsonObject& JsonObject, const TCHAR* ArrayName, const TArray<FPluginDisallowedDescriptor>& Plugins);
+};
\ No newline at end of file
