Fixing loc warning

diff --git a/Engine/Source/Runtime/Projects/Private/PluginDisallowedDescriptor.cpp b/Engine/Source/Runtime/Projects/Private/PluginDisallowedDescriptor.cpp
--- a/Engine/Source/Runtime/Projects/Private/PluginDisallowedDescriptor.cpp
+++ b/Engine/Source/Runtime/Projects/Private/PluginDisallowedDescriptor.cpp
@@ -118,4 +118,6 @@ void FPluginDisallowedDescriptor::UpdateArray(FJsonObject& JsonObject, const TCH
 		FPluginRefJsonArrayUpdater::FGetElementKey::CreateStatic(PluginDisallowedDescriptor::GetPluginRefKey),
 		FPluginRefJsonArrayUpdater::FTryGetJsonObjectKey::CreateStatic(PluginDisallowedDescriptor::TryGetPluginRefJsonObjectKey),
 		FPluginRefJsonArrayUpdater::FUpdateJsonObject::CreateStatic(PluginDisallowedDescriptor::UpdatePluginRefJsonObject));
-}
\ No newline at end of file
+}
+
+#undef LOCTEXT_NAMESPACE
\ No newline at end of file
