Fixed crash when .uproject contains bad Additional[Plugin/Root]Directories values.

diff --git a/Engine/Source/Runtime/Projects/Private/ProjectDescriptor.cpp b/Engine/Source/Runtime/Projects/Private/ProjectDescriptor.cpp
--- a/Engine/Source/Runtime/Projects/Private/ProjectDescriptor.cpp
+++ b/Engine/Source/Runtime/Projects/Private/ProjectDescriptor.cpp
@@ -355,10 +355,16 @@ const FString FProjectDescriptor::MakePathRelativeToProject(const FString& Dir,
 bool FProjectDescriptor::AddPluginDirectory(const FString& Dir)
 {
 #if WITH_EDITOR
-	checkf(!FPaths::IsRelative(Dir), TEXT("%s is not an absolute path"), *Dir);
+	if (!ensureMsgf(!FPaths::IsRelative(Dir), TEXT("Cannot add plugin directory: %s is not an absolute path"), *Dir))
+	{
+		return false;
+	}
 #endif
-	check(!Dir.StartsWith(IFileManager::Get().ConvertToAbsolutePathForExternalAppForWrite(*FPaths::ProjectPluginsDir())));
-	check(!Dir.StartsWith(IFileManager::Get().ConvertToAbsolutePathForExternalAppForWrite(*FPaths::EnginePluginsDir())));
+	if (Dir.StartsWith(IFileManager::Get().ConvertToAbsolutePathForExternalAppForWrite(*FPaths::ProjectPluginsDir())) ||
+		Dir.StartsWith(IFileManager::Get().ConvertToAbsolutePathForExternalAppForWrite(*FPaths::EnginePluginsDir())))
+	{
+		return false;
+	}
 
 	if (!AdditionalPluginDirectories.Contains(Dir))
 	{
@@ -370,15 +376,24 @@ bool FProjectDescriptor::AddPluginDirectory(const FString& Dir)
 
 bool FProjectDescriptor::RemovePluginDirectory(const FString& Dir)
 {
-	checkf(!FPaths::IsRelative(Dir), TEXT("%s is not an absolute path"), *Dir);
+	if (!ensureMsgf(!FPaths::IsRelative(Dir), TEXT("Cannot remove plugin directory: %s is not an absolute path"), *Dir))
+	{
+		return false;
+	}
 	return AdditionalPluginDirectories.RemoveSingle(Dir) > 0;
 }
 
 bool FProjectDescriptor::AddRootDirectory(const FString& Dir)
 {
-	checkf(!FPaths::IsRelative(Dir), TEXT("%s is not an absolute path"), *Dir);
-	check(!Dir.StartsWith(IFileManager::Get().ConvertToAbsolutePathForExternalAppForWrite(*FPaths::EngineDir())));
-	check(!Dir.StartsWith(IFileManager::Get().ConvertToAbsolutePathForExternalAppForWrite(*FPaths::ProjectDir())));
+	if (!ensureMsgf(!FPaths::IsRelative(Dir), TEXT("Cannot add root directory: %s is not an absolute path"), *Dir))
+	{
+		return false;
+	}
+	if (Dir.StartsWith(IFileManager::Get().ConvertToAbsolutePathForExternalAppForWrite(*FPaths::EngineDir())) ||
+		Dir.StartsWith(IFileManager::Get().ConvertToAbsolutePathForExternalAppForWrite(*FPaths::ProjectDir())))
+	{
+		return false;
+	}
 
 	if (!AdditionalRootDirectories.Contains(Dir))
 	{
@@ -390,7 +405,10 @@ bool FProjectDescriptor::AddRootDirectory(const FString& Dir)
 
 bool FProjectDescriptor::RemoveRootDirectory(const FString& Dir)
 {
-	checkf(!FPaths::IsRelative(Dir), TEXT("%s is not an absolute path"), *Dir);
+	if (!ensureMsgf(!FPaths::IsRelative(Dir), TEXT("Cannot remove root directory: %s is not an absolute path"), *Dir))
+	{
+		return false;
+	}
 	return AdditionalRootDirectories.RemoveSingle(Dir) > 0;
 }
 
