Improved plugin manager exclusion support when using -EnableAllPlugins

* Added -ExceptRestrictedPlugins to allow you to exclude any restricted plugins
* Added wildcard support to -ExceptPlugins (and -EnablePlugins and -DisablePlugins) so you can do something like -ExceptPlugins="Online*"

diff --git a/Engine/Source/Runtime/Projects/Private/PluginManager.cpp b/Engine/Source/Runtime/Projects/Private/PluginManager.cpp
--- a/Engine/Source/Runtime/Projects/Private/PluginManager.cpp
+++ b/Engine/Source/Runtime/Projects/Private/PluginManager.cpp
@@ -5,6 +5,7 @@
 #include "GenericPlatform/GenericPlatformFile.h"
 #include "HAL/PlatformFileManager.h"
 #include "HAL/FileManager.h"
+#include "Misc/AsciiSet.h"
 #include "Misc/DataDrivenPlatformInfoRegistry.h"
 #include "Misc/MessageDialog.h"
 #include "Misc/CommandLine.h"
@@ -1352,12 +1353,48 @@ bool FPluginManager::ConfigureEnabledPlugins()
 		{
 			SCOPED_BOOT_TIMING("ParseCmdLineForPlugins");
 
-			auto ParsePluginsList = [](const TCHAR* InListKey) -> TArray<FString>
+			auto ParsePluginsList = [this](const TCHAR* InListKey) -> TArray<FString>
 			{
 				TArray<FString> PluginsList;
-				FString PluginsListStr;
-				FParse::Value(FCommandLine::Get(), InListKey, PluginsListStr, false);
-				PluginsListStr.ParseIntoArray(PluginsList, TEXT(","));
+				{
+					FString PluginsListStr;
+					FParse::Value(FCommandLine::Get(), InListKey, PluginsListStr, false);
+					PluginsListStr.ParseIntoArray(PluginsList, TEXT(","));
+				}
+
+				TArray<FString> WildcardPluginsList;
+				PluginsList.RemoveAll([&WildcardPluginsList](const FString& ExceptPluginName)
+				{
+					constexpr FAsciiSet Wildcards("*?");
+					if (FAsciiSet::HasAny(ExceptPluginName, Wildcards))
+					{
+						WildcardPluginsList.Add(ExceptPluginName);
+						return true;
+					}
+					return false;
+				});
+
+				if (WildcardPluginsList.Num() > 0)
+				{
+					for (const FString& PotentialPluginName : PluginsToConfigure)
+					{
+						bool bMatchesAnyWildcard = false;
+						for (const FString& WildcardPluginName : WildcardPluginsList)
+						{
+							if (PotentialPluginName.MatchesWildcard(WildcardPluginName))
+							{
+								bMatchesAnyWildcard = true;
+								break;
+							}
+						}
+
+						if (bMatchesAnyWildcard)
+						{
+							PluginsList.Add(PotentialPluginName);
+						}
+					}
+				}
+
 				return PluginsList;
 			};
 
@@ -1374,10 +1411,21 @@ bool FPluginManager::ConfigureEnabledPlugins()
 			}
 			if (ExtraPluginsToEnable.Num() > 0)
 			{
+				auto IsRestrictedPlugin = [this](const FString& PluginName)
+				{
+					if (TSharedPtr<IPlugin> PluginPtr = FindPlugin(PluginName))
+					{
+						const FString& PluginBaseDir = PluginPtr->GetBaseDir();
+						return PluginBaseDir.Contains(TEXT("/Restricted/"));
+					}
+					return true;
+				};
+
 				const TArray<FString> ExceptPlugins = ParsePluginsList(TEXT("ExceptPlugins="));
+				const bool bExceptRestrictedPlugins = FParse::Param(FCommandLine::Get(), TEXT("ExceptRestrictedPlugins"));
 				for (const FString& EnablePluginName : ExtraPluginsToEnable)
 				{
-					if (!ConfiguredPluginNames.Contains(EnablePluginName) && !ExceptPlugins.Contains(EnablePluginName))
+					if (!ConfiguredPluginNames.Contains(EnablePluginName) && !ExceptPlugins.Contains(EnablePluginName) && (!bExceptRestrictedPlugins || !IsRestrictedPlugin(EnablePluginName)))
 					{
 						if (!ConfigureEnabledPluginForCurrentTarget(FPluginReferenceDescriptor(EnablePluginName, true), EnabledPlugins))
 						{
@@ -2455,7 +2503,14 @@ bool FPluginManager::TryLoadModulesForPlugin( const FPlugin& Plugin, const ELoad
 
 	if( !FailureMessage.IsEmpty() )
 	{
-		UE_LOG(LogPluginManager, Error, TEXT("%s"), *FailureMessage.ToString());
+		if (bAllPluginsEnabledViaCommandLine)
+		{
+			UE_LOG(LogPluginManager, Warning, TEXT("%s"), *FailureMessage.ToString());
+		}
+		else
+		{
+			UE_LOG(LogPluginManager, Error, TEXT("%s"), *FailureMessage.ToString());
+		}
 		if (!bAllPluginsEnabledViaCommandLine)
 		{
 			FMessageDialog::Open(EAppMsgType::Ok, FailureMessage);
