Downgrade warning to display to avoid CIS errors when running with -EnableAllPlugins

#rnx

diff --git a/Engine/Source/Runtime/Projects/Private/PluginManager.cpp b/Engine/Source/Runtime/Projects/Private/PluginManager.cpp
--- a/Engine/Source/Runtime/Projects/Private/PluginManager.cpp
+++ b/Engine/Source/Runtime/Projects/Private/PluginManager.cpp
@@ -2505,7 +2505,7 @@ bool FPluginManager::TryLoadModulesForPlugin( const FPlugin& Plugin, const ELoad
 	{
 		if (bAllPluginsEnabledViaCommandLine)
 		{
-			UE_LOG(LogPluginManager, Warning, TEXT("%s"), *FailureMessage.ToString());
+			UE_LOG(LogPluginManager, Display, TEXT("%s"), *FailureMessage.ToString());
 		}
 		else
 		{
