IPlugin::SetVersePath

diff --git a/Engine/Source/Runtime/Projects/Private/PluginManager.h b/Engine/Source/Runtime/Projects/Private/PluginManager.h
--- a/Engine/Source/Runtime/Projects/Private/PluginManager.h
+++ b/Engine/Source/Runtime/Projects/Private/PluginManager.h
@@ -100,6 +100,11 @@ public:
 		return Descriptor.VersePath;
 	}
 
+	virtual void SetVersePath(FString&& InVersePath) override
+	{
+		Descriptor.VersePath = MoveTemp(InVersePath);
+	}
+
 	virtual EPluginType GetType() const override
 	{
 		return Type;
diff --git a/Engine/Source/Runtime/Projects/Public/Interfaces/IPluginManager.h b/Engine/Source/Runtime/Projects/Public/Interfaces/IPluginManager.h
--- a/Engine/Source/Runtime/Projects/Public/Interfaces/IPluginManager.h
+++ b/Engine/Source/Runtime/Projects/Public/Interfaces/IPluginManager.h
@@ -171,6 +171,12 @@ public:
 	 */
 	virtual const FString& GetVersePath() const = 0;
 
+	/**
+	 * Sets the Verse path to the root of the plugin's content directory
+	 * @param InVersePath Verse path to set
+	 */
+	virtual void SetVersePath(FString&& InVersePath) = 0;
+
 	/**
 	 * Returns the plugin's location
 	 *
