FPluginManager::UnmountExplicitlyLoadedPlugin can unmount code modules

diff --git a/Engine/Source/Runtime/Core/Public/Modules/ModuleManager.h b/Engine/Source/Runtime/Core/Public/Modules/ModuleManager.h
--- a/Engine/Source/Runtime/Core/Public/Modules/ModuleManager.h
+++ b/Engine/Source/Runtime/Core/Public/Modules/ModuleManager.h
@@ -63,6 +63,17 @@ enum class EModuleLoadResult
 	FailedToInitialize
 };
 
+/**
+ * Enumerates reasons for failed module unloads.
+ */
+enum class EModuleUnloadResult
+{
+	/** Module unloaded successfully. */
+	Success,
+
+	/** Module does not support dynamic reloading (see IModuleInterface::SupportsDynamicReloading). */
+	UnloadNotSupported
+};
 
 /**
  * Enumerates reasons for modules to change.
diff --git a/Engine/Source/Runtime/Projects/Private/ModuleDescriptor.cpp b/Engine/Source/Runtime/Projects/Private/ModuleDescriptor.cpp
--- a/Engine/Source/Runtime/Projects/Private/ModuleDescriptor.cpp
+++ b/Engine/Source/Runtime/Projects/Private/ModuleDescriptor.cpp
@@ -703,6 +703,41 @@ void FModuleDescriptor::LoadModulesForPhase(ELoadingPhase::Type LoadingPhase, co
 	}
 }
 
+void FModuleDescriptor::UnloadModulesForPhase(ELoadingPhase::Type LoadingPhase, const TArray<FModuleDescriptor>& Modules, TMap<FName, EModuleUnloadResult>& OutErrors, bool bSkipUnload /*= false*/)
+{
+	FScopedSlowTask SlowTask((float)Modules.Num());
+	for (const FModuleDescriptor& Descriptor : Modules)
+	{
+		SlowTask.EnterProgressFrame();
+
+		if (LoadingPhase != Descriptor.LoadingPhase)
+		{
+			continue;
+		}
+
+		IModuleInterface* Module = FModuleManager::Get().GetModule(Descriptor.Name);
+		if (!Module)
+		{
+			continue;
+		}
+
+		if (!Module->SupportsDynamicReloading())
+		{
+			OutErrors.Add(Descriptor.Name, EModuleUnloadResult::UnloadNotSupported);
+			continue;
+		}
+
+		if (bSkipUnload)
+		{
+			// Useful to gather errors without actually unloading
+			continue;
+		}
+
+		Module->PreUnloadCallback();
+		verify(FModuleManager::Get().UnloadModule(Descriptor.Name));
+	}
+}
+
 #if !IS_MONOLITHIC
 bool FModuleDescriptor::CheckModuleCompatibility(const TArray<FModuleDescriptor>& Modules, TArray<FString>& OutIncompatibleFiles)
 {
diff --git a/Engine/Source/Runtime/Projects/Private/PluginManager.cpp b/Engine/Source/Runtime/Projects/Private/PluginManager.cpp
--- a/Engine/Source/Runtime/Projects/Private/PluginManager.cpp
+++ b/Engine/Source/Runtime/Projects/Private/PluginManager.cpp
@@ -2516,6 +2516,54 @@ bool FPluginManager::TryLoadModulesForPlugin( const FPlugin& Plugin, const ELoad
 	return true;
 }
 
+bool FPluginManager::TryUnloadModulesForPlugin(const FPlugin& Plugin, const ELoadingPhase::Type LoadingPhase, FText* OutFailureMessage /*= nullptr*/, bool bSkipUnload /*= false*/) const
+{
+	TMap<FName, EModuleUnloadResult> Errors;
+	FModuleDescriptor::UnloadModulesForPhase(LoadingPhase, Plugin.Descriptor.Modules, Errors, bSkipUnload);
+
+	FText FailureMessage;
+	for( const TPair<FName, EModuleUnloadResult>& FailureIt : Errors)
+	{
+		const FName ModuleNameThatFailedToLoad = FailureIt.Key;
+		const EModuleUnloadResult FailureReason = FailureIt.Value;
+
+		if (FailureReason != EModuleUnloadResult::Success)
+		{
+			const FText PluginNameText = FText::FromString(Plugin.Name);
+			const FText TextModuleName = FText::FromName(ModuleNameThatFailedToLoad);
+
+			if (FailureReason == EModuleUnloadResult::UnloadNotSupported)
+			{
+				FailureMessage = FText::Format(LOCTEXT("UnloadNotSupported", "Plugin '{0}' failed to unload because module '{1}' does not support unloading."), PluginNameText, TextModuleName);
+			}
+			else
+			{
+				ensure(0);	// If this goes off, the error handling code should be updated for the new enum values!
+				FailureMessage = FText::Format(LOCTEXT("PluginGenericUnloadFailure", "Plugin '{0}' failed to unload because module '{1}' could not be unloaded for an unspecified reason. Please report this error."), PluginNameText, TextModuleName);
+			}
+
+			// Don't need to display more than one module load error per plugin that failed to load
+			break;
+		}
+	}
+
+	if (!FailureMessage.IsEmpty())
+	{
+		UE_LOG(LogPluginManager, Error, TEXT("%s"), *FailureMessage.ToString());
+
+		FMessageDialog::Open(EAppMsgType::Ok, FailureMessage);
+
+		if (OutFailureMessage)
+		{
+			*OutFailureMessage = MoveTemp(FailureMessage);
+		}
+
+		return false;
+	}
+
+	return true;
+}
+
 bool FPluginManager::LoadModulesForEnabledPlugins( const ELoadingPhase::Type LoadingPhase )
 {
 	// Figure out which plugins are enabled
@@ -3048,13 +3096,25 @@ bool FPluginManager::UnmountExplicitlyLoadedPlugin(const FString& PluginName, FT
 		return false;
 	}
 
-	if (Plugin->Descriptor.Modules.Num() > 0)
+	// Simulate unload of all the plugin modules to gather errors
+	// We don't want to actually unload modules until content has been unloaded
+	for (ELoadingPhase::Type LoadingPhase = (ELoadingPhase::Type)(ELoadingPhase::Max - 1);
+		LoadingPhase >= (ELoadingPhase::Type)0;
+		LoadingPhase = (ELoadingPhase::Type)(LoadingPhase - 1))
 	{
-		if (OutReason)
+		if (LoadingPhase != ELoadingPhase::None)
 		{
-			*OutReason = LOCTEXT("UnloadPluginContainedModules", "Plugin contains modules and may be unsafe to unload");
+			constexpr bool bSkipUnload = true;
+			FText FailureMessage;
+			if (!TryUnloadModulesForPlugin(*Plugin, LoadingPhase, &FailureMessage, bSkipUnload))
+			{
+				if (OutReason)
+				{
+					*OutReason = MoveTemp(FailureMessage);
+				}
+				return false;
+			}
 		}
-		return false;
 	}
 
 	// Notify that additional localization data should be unloaded
@@ -3089,6 +3149,17 @@ bool FPluginManager::UnmountExplicitlyLoadedPlugin(const FString& PluginName, FT
 		}
 	}
 
+	// Actually unload all the plugin modules now that content unmount is finished
+	for (ELoadingPhase::Type LoadingPhase = (ELoadingPhase::Type)(ELoadingPhase::Max - 1); 
+		 LoadingPhase >= (ELoadingPhase::Type)0; 
+		 LoadingPhase = (ELoadingPhase::Type)(LoadingPhase - 1))
+	{
+		if (LoadingPhase != ELoadingPhase::None)
+		{
+			verify(TryUnloadModulesForPlugin(*Plugin, LoadingPhase));
+		}
+	}
+
 	Plugin->SetIsMounted(false);
 	Plugin->bEnabled = false;
 
diff --git a/Engine/Source/Runtime/Projects/Private/PluginManager.h b/Engine/Source/Runtime/Projects/Private/PluginManager.h
--- a/Engine/Source/Runtime/Projects/Private/PluginManager.h
+++ b/Engine/Source/Runtime/Projects/Private/PluginManager.h
@@ -259,6 +259,9 @@ private:
 	/** Attempt to load all the modules for the given plugin */
 	bool TryLoadModulesForPlugin(const FPlugin& Plugin, const ELoadingPhase::Type LoadingPhase) const;
 
+	/** Attempt to unload all the modules for the given plugin */
+	bool TryUnloadModulesForPlugin(const FPlugin& Plugin, const ELoadingPhase::Type LoadingPhase, FText* OutFailureMessage = nullptr, bool bSkipUnload = false) const;
+
 	/** Gets the instance of a given plugin */
 	TSharedPtr<FPlugin> FindPluginInstance(const FString& Name);
 
diff --git a/Engine/Source/Runtime/Projects/Public/ModuleDescriptor.h b/Engine/Source/Runtime/Projects/Public/ModuleDescriptor.h
--- a/Engine/Source/Runtime/Projects/Public/ModuleDescriptor.h
+++ b/Engine/Source/Runtime/Projects/Public/ModuleDescriptor.h
@@ -16,6 +16,7 @@ class FJsonObject;
 class FText;
 
 enum class EModuleLoadResult;
+enum class EModuleUnloadResult;
 
 /**
  * Phase at which this module should be loaded during startup.
@@ -231,6 +232,9 @@ struct FModuleDescriptor
 	/** Loads all the modules for a given loading phase. Returns a map of module names to load errors */
 	static PROJECTS_API void LoadModulesForPhase(ELoadingPhase::Type LoadingPhase, const TArray<FModuleDescriptor>& Modules, TMap<FName, EModuleLoadResult>& ModuleLoadErrors);
 
+	/** Unloads all the modules for a given loading phase. Returns a map of module names to load errors. bSkipUnload can be used to simulate unloading */
+	static PROJECTS_API void UnloadModulesForPhase(ELoadingPhase::Type LoadingPhase, const TArray<FModuleDescriptor>& Modules, TMap<FName, EModuleUnloadResult>& OutErrors, bool bSkipUnload = false);
+
 #if !IS_MONOLITHIC
 	/** Checks that all modules are compatible with the current engine version. Returns false and appends a list of names to OutIncompatibleFiles if not. */
 	static PROJECTS_API bool CheckModuleCompatibility(const TArray<FModuleDescriptor>& Modules, TArray<FString>& OutIncompatibleFiles);
