Process all configs related to a plugin when it is loaded, even if the related configs were discovered before the plugin was loaded

#rnx

diff --git a/Engine/Source/Runtime/Projects/Private/PluginManager.cpp b/Engine/Source/Runtime/Projects/Private/PluginManager.cpp
--- a/Engine/Source/Runtime/Projects/Private/PluginManager.cpp
+++ b/Engine/Source/Runtime/Projects/Private/PluginManager.cpp
@@ -1785,19 +1785,49 @@ bool FPluginManager::ConfigureEnabledPlugins()
 			FCriticalSection ConfigCS;
 			FCriticalSection PluginPakCS;
 			
+			struct FPendingConfigFile
+			{
+				FString PluginName;
+				FString PluginConfigDir;
+				FString PluginConfigFile;
+			};
+			FCriticalSection PendingConfigsCS;
+			TMap<FString, TArray<FPendingConfigFile>> PendingConfigs;
+
 			// Mount all the enabled plugins
-			ParallelFor(PluginsArray.Num(), [&PluginsArray, &ConfigCS, &PluginPakCS, &ConfigFilesPluginsCannotOverride, &AllIniFiles, this](int32 Index)
+			ParallelFor(PluginsArray.Num(), [&PluginsArray, &ConfigCS, &PluginPakCS, &PendingConfigsCS, &PendingConfigs, &ConfigFilesPluginsCannotOverride, &AllIniFiles, this](int32 Index)
 			{
 				FString PlatformName = FPlatformProperties::PlatformName();
 				FPlugin& Plugin = *PluginsArray[Index];
 				UE_LOG(LogPluginManager, Log, TEXT("Mounting %s plugin %s"), *EnumToString(Plugin.Type), *Plugin.GetName());
 				UE_LOG(LogPluginManager, Verbose, TEXT("Plugin path: %s"), *Plugin.FileName);
 
+				auto AppendPluginConfigData = [&ConfigFilesPluginsCannotOverride](FConfigFile& DestinationPluginConfig, const FString& DestinationPluginConfigFilename, const FString& SourcePluginName, const FString& SourcePluginConfigDir, const FString& SourcePluginConfigFile)
+				{
+					UE_LOG(LogPluginManager, Log, TEXT("Found config from plugin[%s] %s"), *SourcePluginName, *DestinationPluginConfigFilename);
+
+					FString BaseConfigFile = *FPaths::GetBaseFilename(SourcePluginConfigFile);
+					if (ConfigFilesPluginsCannotOverride.Contains(BaseConfigFile))
+					{
+						// Not allowed, skip it
+						FText FailureMessage = FText::Format(LOCTEXT("PluginOverrideFailureFormat", "Plugin '{0}' cannot override config file: '{1}'"), FText::FromString(SourcePluginName), FText::FromString(BaseConfigFile));
+						FText DialogTitle = LOCTEXT("PluginConfigFileOverride", "Plugin config file override");
+						UE_LOG(LogPluginManager, Error, TEXT("%s"), *FailureMessage.ToString());
+						FMessageDialog::Open(EAppMsgType::Ok, FailureMessage, DialogTitle);
+						return;
+					}
+
+					DestinationPluginConfig.AddDynamicLayerToHierarchy(FPaths::Combine(SourcePluginConfigDir, SourcePluginConfigFile));
+
+#if ALLOW_INI_OVERRIDE_FROM_COMMANDLINE
+					// Don't allow plugins to stomp command line overrides, so re-apply them
+					FConfigFile::OverrideFromCommandline(&DestinationPluginConfig, DestinationPluginConfigFilename);
+#endif // ALLOW_INI_OVERRIDE_FROM_COMMANDLINE
+				};
 
 				// Build the config system key for PluginName.ini
 				FString PluginConfigFilename = GConfig->GetConfigFilename(*Plugin.Name);
 				{
-					
 					FScopeLock Locker(&ConfigCS);
 
 					FConfigFile& PluginConfig = GConfig->Add(PluginConfigFilename, FConfigFile());
@@ -1809,7 +1839,19 @@ bool FPluginManager::ConfigureEnabledPlugins()
 						Context.IniCacheSet = &AllIniFiles;
 					}
 
-					if (!Context.Load(*Plugin.Name))
+					if (Context.Load(*Plugin.Name))
+					{
+						// Process anything relevant that was discovered before we loaded
+						FScopeLock PendingConfigsLock(&PendingConfigsCS);
+						if (const TArray<FPendingConfigFile>* PendingConfigArray = PendingConfigs.Find(Plugin.Name))
+						{
+							for (const FPendingConfigFile& PendingConfigFile : *PendingConfigArray)
+							{
+								AppendPluginConfigData(PluginConfig, PluginConfigFilename, PendingConfigFile.PluginName, PendingConfigFile.PluginConfigDir, PendingConfigFile.PluginConfigFile);
+							}
+						}
+					}
+					else
 					{
 						// Nothing to add, remove from map
 						GConfig->Remove(PluginConfigFilename);
@@ -1828,7 +1870,7 @@ bool FPluginManager::ConfigureEnabledPlugins()
 
 					if (BaseConfigFile == Plugin.Name)
 					{
-						// We just handles this, skip it
+						// We just handled this, skip it
 						continue;
 					}
 
@@ -1840,24 +1882,12 @@ bool FPluginManager::ConfigureEnabledPlugins()
 
 						if (FoundConfig != nullptr)
 						{
-							UE_LOG(LogPluginManager, Log, TEXT("Found config from plugin[%s] %s"), *Plugin.GetName(), *PluginConfigFilename);
-
-							if (ConfigFilesPluginsCannotOverride.Contains(BaseConfigFile))
-							{
-								// Not allowed, skip it
-								FText FailureMessage = FText::Format(LOCTEXT("PluginOverrideFailureFormat", "Plugin '{0}' cannot override config file: '{1}'"), FText::FromString(Plugin.GetName()), FText::FromString(BaseConfigFile));
-								FText DialogTitle = LOCTEXT("PluginConfigFileOverride", "Plugin config file override");
-								UE_LOG(LogPluginManager, Error, TEXT("%s"), *FailureMessage.ToString());
-								FMessageDialog::Open(EAppMsgType::Ok, FailureMessage, DialogTitle);
-								continue;
-							}
-
-							FoundConfig->AddDynamicLayerToHierarchy(FPaths::Combine(PluginConfigDir, ConfigFile));
-
-#if ALLOW_INI_OVERRIDE_FROM_COMMANDLINE
-							// Don't allow plugins to stomp command line overrides, so re-apply them
-							FConfigFile::OverrideFromCommandline(FoundConfig, PluginConfigFilename);
-#endif // ALLOW_INI_OVERRIDE_FROM_COMMANDLINE
+							AppendPluginConfigData(*FoundConfig, PluginConfigFilename, Plugin.GetName(), PluginConfigDir, ConfigFile);
+						}
+						else if (PluginsToConfigure.Contains(BaseConfigFile))
+						{
+							FScopeLock PendingConfigsLock(&PendingConfigsCS);
+							PendingConfigs.FindOrAdd(BaseConfigFile).Add(FPendingConfigFile{ Plugin.GetName(), PluginConfigDir, ConfigFile });
 						}
 					}
 				}
