Fix FPluginDescriptor::UpdateJson ignoring VerseVersion field, and not removing other Verse-related fields when appropriate.

diff --git a/Engine/Source/Runtime/Projects/Private/PluginDescriptor.cpp b/Engine/Source/Runtime/Projects/Private/PluginDescriptor.cpp
--- a/Engine/Source/Runtime/Projects/Private/PluginDescriptor.cpp
+++ b/Engine/Source/Runtime/Projects/Private/PluginDescriptor.cpp
@@ -386,16 +386,37 @@ void FPluginDescriptor::UpdateJson(FJsonObject& JsonObject) const
 	{
 		JsonObject.SetStringField(TEXT("VersePath"), VersePath);
 	}
+	else
+	{
+		JsonObject.RemoveField(TEXT("VersePath"));
+	}
 
 	if (VerseScope != EVerseScope::PublicUser)
 	{
 		JsonObject.SetStringField(TEXT("VerseScope"), EVerseScope::ToString(VerseScope));
 	}
+	else
+	{
+		JsonObject.RemoveField(TEXT("VerseScope"));
+	}
+
+	if (VerseVersion.IsSet())
+	{
+		JsonObject.SetNumberField(TEXT("VerseVersion"), VerseVersion.GetValue());
+	}
+	else
+	{
+		JsonObject.RemoveField(TEXT("VerseVersion"));
+	}
 
 	if (bEnableVerseAssetReflection)
 	{
 		JsonObject.SetBoolField(TEXT("EnableVerseAssetReflection"), bEnableVerseAssetReflection);
 	}
+	else
+	{
+		JsonObject.RemoveField(TEXT("EnableVerseAssetReflection"));
+	}
 
 	if (EnabledByDefault != EPluginEnabledByDefault::Unspecified)
 	{
@@ -411,6 +432,11 @@ void FPluginDescriptor::UpdateJson(FJsonObject& JsonObject) const
 	{
 		JsonObject.SetBoolField(TEXT("CanContainVerse"), bCanContainVerse);
 	}
+	else
+	{
+		JsonObject.RemoveField(TEXT("CanContainVerse"));
+	}
+
 	if (bNoCode)
 	{
 		JsonObject.SetBoolField(TEXT("NoCode"), bNoCode);
