Plugin descriptor load/save api that takes in a TCHAR* for the filename instead of an FString
#rnx

diff --git a/Engine/Source/Runtime/Projects/Private/PluginDescriptor.cpp b/Engine/Source/Runtime/Projects/Private/PluginDescriptor.cpp
--- a/Engine/Source/Runtime/Projects/Private/PluginDescriptor.cpp
+++ b/Engine/Source/Runtime/Projects/Private/PluginDescriptor.cpp
@@ -11,26 +11,26 @@
 
 namespace PluginDescriptor
 {
-	bool ReadFile(const FString& FileName, FString& Text, FText* OutFailReason = nullptr)
+	bool ReadFile(const TCHAR* FileName, FString& Text, FText* OutFailReason = nullptr)
 	{
-		if (!FFileHelper::LoadFileToString(Text, *FileName))
+		if (!FFileHelper::LoadFileToString(Text, FileName))
 		{
 			if (OutFailReason)
 			{
-				*OutFailReason = FText::Format(LOCTEXT("FailedToLoadDescriptorFile", "Failed to open descriptor file '{0}'"), FText::FromString(FileName));
+				*OutFailReason = FText::Format(LOCTEXT("FailedToLoadDescriptorFile", "Failed to open descriptor file '{0}'"), FText::FromStringView(FileName));
 			}
 			return false;
 		}
 		return true;
 	}
 
-	bool WriteFile(const FString& FileName, const FString& Text, FText* OutFailReason = nullptr)
+	bool WriteFile(const TCHAR* FileName, const FString& Text, FText* OutFailReason = nullptr)
 	{
-		if (!FFileHelper::SaveStringToFile(Text, *FileName))
+		if (!FFileHelper::SaveStringToFile(Text, FileName))
 		{
 			if (OutFailReason)
 			{
-				*OutFailReason = FText::Format(LOCTEXT("FailedToWriteDescriptorFile", "Failed to write plugin descriptor file '{0}'. Perhaps the file is Read-Only?"), FText::FromString(FileName));
+				*OutFailReason = FText::Format(LOCTEXT("FailedToWriteDescriptorFile", "Failed to write plugin descriptor file '{0}'. Perhaps the file is Read-Only?"), FText::FromStringView(FileName));
 			}
 			return false;
 		}
@@ -94,8 +94,7 @@ FPluginDescriptor::FPluginDescriptor()
 {
 }
 
-
-bool FPluginDescriptor::Load(const FString& FileName, FText* OutFailReason /*= nullptr*/)
+bool FPluginDescriptor::Load(const TCHAR* FileName, FText* OutFailReason /*= nullptr*/)
 {
 #if WITH_EDITOR
 	CachedJson.Reset();
@@ -110,9 +109,14 @@ bool FPluginDescriptor::Load(const FString& FileName, FText* OutFailReason /*= n
 	return false;
 }
 
+bool FPluginDescriptor::Load(const FString& FileName, FText* OutFailReason /*= nullptr*/)
+{
+	return Load(*FileName, OutFailReason);
+}
+
 bool FPluginDescriptor::Load(const FString& FileName, FText& OutFailReason)
 {
-	return Load(FileName, &OutFailReason);
+	return Load(*FileName, &OutFailReason);
 }
 
 bool FPluginDescriptor::Read(const FString& Text, FText* OutFailReason /*= nullptr*/)
@@ -310,7 +314,7 @@ bool FPluginDescriptor::Read(const FJsonObject& Object, FText& OutFailReason)
 	return Read(Object, &OutFailReason);
 }
 
-bool FPluginDescriptor::Save(const FString& FileName, FText* OutFailReason /*= nullptr*/) const
+bool FPluginDescriptor::Save(const TCHAR* FileName, FText* OutFailReason /*= nullptr*/) const
 {
 	// Write the descriptor to text
 	FString Text;
@@ -320,9 +324,14 @@ bool FPluginDescriptor::Save(const FString& FileName, FText* OutFailReason /*= n
 	return PluginDescriptor::WriteFile(FileName, Text, OutFailReason);
 }
 
+bool FPluginDescriptor::Save(const FString& FileName, FText* OutFailReason /*= nullptr*/) const
+{
+	return Save(*FileName, OutFailReason);
+}
+
 bool FPluginDescriptor::Save(const FString& FileName, FText& OutFailReason) const
 {
-	return Save(FileName, &OutFailReason);
+	return Save(*FileName, &OutFailReason);
 }
 
 void FPluginDescriptor::Write(FString& Text) const
@@ -553,7 +562,7 @@ bool FPluginDescriptor::UpdatePluginFile(const FString& FileName, FText* OutFail
 		// Plugin file exists so we need to read it and update it.
 
 		FString JsonText;
-		if (!PluginDescriptor::ReadFile(FileName, JsonText, OutFailReason))
+		if (!PluginDescriptor::ReadFile(*FileName, JsonText, OutFailReason))
 		{
 			return false;
 		}
@@ -581,7 +590,7 @@ bool FPluginDescriptor::UpdatePluginFile(const FString& FileName, FText* OutFail
 #if WITH_EDITOR
 		CachedJson = JsonObject;
 #endif
-		return PluginDescriptor::WriteFile(FileName, JsonText, OutFailReason);
+		return PluginDescriptor::WriteFile(*FileName, JsonText, OutFailReason);
 	}
 	else
 	{
diff --git a/Engine/Source/Runtime/Projects/Public/PluginDescriptor.h b/Engine/Source/Runtime/Projects/Public/PluginDescriptor.h
--- a/Engine/Source/Runtime/Projects/Public/PluginDescriptor.h
+++ b/Engine/Source/Runtime/Projects/Public/PluginDescriptor.h
@@ -170,6 +170,9 @@ struct FPluginDescriptor
 	/** Constructor. */
 	PROJECTS_API FPluginDescriptor();
 
+	/** Loads the descriptor from the given file. */
+	PROJECTS_API bool Load(const TCHAR* FileName, FText* OutFailReason = nullptr);
+
 	/** Loads the descriptor from the given file. */
 	PROJECTS_API bool Load(const FString& FileName, FText* OutFailReason = nullptr);
 
@@ -188,10 +191,13 @@ struct FPluginDescriptor
 	/** Reads the descriptor from the given JSON object */
 	PROJECTS_API bool Read(const FJsonObject& Object, FText& OutFailReason);
 
-	/** Saves the descriptor from the given file. */
+	/** Saves the descriptor to the given file. */
+	PROJECTS_API bool Save(const TCHAR* FileName, FText* OutFailReason = nullptr) const;
+
+	/** Saves the descriptor to the given file. */
 	PROJECTS_API bool Save(const FString& FileName, FText* OutFailReason = nullptr) const;
 
-	/** Saves the descriptor from the given file. */
+	/** Saves the descriptor to the given file. */
 	PROJECTS_API bool Save(const FString& FileName, FText& OutFailReason) const;
 
 	/** Writes a descriptor to JSON */
