Add optional plugins to DefaultEngine config

diff --git a/Engine/Source/Runtime/Projects/Private/PluginManager.cpp b/Engine/Source/Runtime/Projects/Private/PluginManager.cpp
--- a/Engine/Source/Runtime/Projects/Private/PluginManager.cpp
+++ b/Engine/Source/Runtime/Projects/Private/PluginManager.cpp
@@ -1610,24 +1610,36 @@ bool FPluginManager::ConfigureEnabledPlugins()
 		}
 
 #if IS_PROGRAM
+		auto EnableProgramPlugin = [this, &ConfiguredPluginNames, &EnabledPlugins](const TCHAR* ConfigEntry, bool bOptionalPlugin) mutable
 		{
-			// Programs can also define the list of enabled plugins in ini
-			SCOPED_BOOT_TIMING("AddProgramEnabledPlugins");
-
 			TArray<FString> ProgramPluginNames;
-			GConfig->GetArray(TEXT("Plugins"), TEXT("ProgramEnabledPlugins"), ProgramPluginNames, GEngineIni);
+			GConfig->GetArray(TEXT("Plugins"), ConfigEntry, ProgramPluginNames, GEngineIni);
 
 			for (const FString& PluginName : ProgramPluginNames)
 			{
 				if (!ConfiguredPluginNames.Contains(PluginName))
 				{
-					if (!ConfigureEnabledPluginForCurrentTarget(FPluginReferenceDescriptor(PluginName, true), EnabledPlugins))
+					FPluginReferenceDescriptor PluginReference(PluginName, true);
+					PluginReference.bOptional = bOptionalPlugin;
+					if (!ConfigureEnabledPluginForCurrentTarget(PluginReference, EnabledPlugins))
 					{
 						return false;
 					}
 					ConfiguredPluginNames.Add(PluginName);
 				}
 			}
+
+			return true;
+		};
+
+		{
+			// Programs can also define the list of enabled plugins in ini
+			SCOPED_BOOT_TIMING("AddProgramEnabledPlugins");
+
+			if (!EnableProgramPlugin(TEXT("ProgramEnabledPlugins"), false) || !EnableProgramPlugin(TEXT("ProgramOptionalPlugins"), true))
+			{
+				return false;
+			}
 		}
 #endif
 
