PluginManager: Improve feedback when a required plugin is unavailable and the process is terminated.
#rnx

diff --git a/Engine/Source/Runtime/Projects/Private/PluginManager.cpp b/Engine/Source/Runtime/Projects/Private/PluginManager.cpp
--- a/Engine/Source/Runtime/Projects/Private/PluginManager.cpp
+++ b/Engine/Source/Runtime/Projects/Private/PluginManager.cpp
@@ -27,6 +27,7 @@
 #include "Async/ParallelFor.h"
 #include "Misc/ScopeRWLock.h"
 #include "Algo/Accumulate.h"
+#include "Algo/Reverse.h"
 #include "Containers/VersePath.h"
 #include "Internationalization/TextLocalizationManager.h"
 #include "ProfilingDebugging/CpuProfilerTrace.h"
@@ -514,6 +515,7 @@ struct FPluginManager::FConfigurePluginResultInfo
 {
 	EConfigurePluginResultCode Code;
 	const FPluginReferenceDescriptor* PluginReference;
+	TArray<const FPluginReferenceDescriptor*> ReferenceChain;
 
 	FConfigurePluginResultInfo()
 		: Code(EConfigurePluginResultCode::Unset)
@@ -759,9 +761,11 @@ void FPluginManager::DiscoverAllPlugins()
 	}
 }
 
-void FPluginManager::ReadAllPlugins(FDiscoveredPluginMap& Plugins, const TSet<FString>& ExtraSearchPaths)
+void FPluginManager::ReadAllPlugins(FDiscoveredPluginMap& Plugins, const TSet<FString>& ExtraSearchPaths,
+	TArray<FString>* OutPluginSources)
 {
 #if (WITH_ENGINE && !IS_PROGRAM) || (WITH_PLUGIN_SUPPORT && !UE_DISABLE_PLUGIN_DISCOVERY)
+	TArray<FString> OptionalOutPluginRoots;
 	const FProjectDescriptor* Project = IProjectManager::Get().GetCurrentProject();
 
 	// Find any plugin manifest files. These give us the plugin list (and their descriptors) without needing to scour the directory tree.
@@ -786,6 +790,10 @@ void FPluginManager::ReadAllPlugins(FDiscoveredPluginMap& Plugins, const TSet<FS
 		SlowTask_ReadAll.EnterProgressFrame(0.5f);
 		// Find "built-in" plugins.  That is, plugins situated right within the Engine directory.
 		TArray<FString> EnginePluginDirs = FPaths::GetExtensionDirs(FPaths::EngineDir(), TEXT("Plugins"), !GIsEditor);
+		if (OutPluginSources)
+		{
+			OptionalOutPluginRoots.Append(EnginePluginDirs);
+		}
 		{
 			FScopedSlowTask SlowTask_ReadEngine((float)EnginePluginDirs.Num());
 			
@@ -802,6 +810,10 @@ void FPluginManager::ReadAllPlugins(FDiscoveredPluginMap& Plugins, const TSet<FS
 		if (Project != nullptr)
 		{
 			TArray<FString> ProjectPluginDirs = FPaths::GetExtensionDirs(FPaths::GetPath(FPaths::GetProjectFilePath()), TEXT("Plugins"), !GIsEditor);
+			if (OutPluginSources)
+			{
+				OptionalOutPluginRoots.Append(ProjectPluginDirs);
+			}
 			FScopedSlowTask SlowTask_ReadProject((float)ProjectPluginDirs.Num());
 
 			for (const FString& ProjectPluginDir : ProjectPluginDirs)
@@ -822,6 +834,11 @@ void FPluginManager::ReadAllPlugins(FDiscoveredPluginMap& Plugins, const TSet<FS
 			SlowTask_ReadManifest.EnterProgressFrame();
 
 			UE_LOG(LogPluginManager, Verbose, TEXT("Reading plugin manifest: %s"), *ManifestFileName);
+			if (OutPluginSources)
+			{
+				OutPluginSources->Add(FString::Printf(TEXT("Manifest: %s"),
+					*FPaths::ConvertRelativePathToFull(ManifestFileName)));
+			}
 			FPluginManifest Manifest;
 
 			// Try to load the manifest. We only expect manifests in a cooked game, so failing to load them is a hard error.
@@ -872,11 +889,19 @@ void FPluginManager::ReadAllPlugins(FDiscoveredPluginMap& Plugins, const TSet<FS
 		SlowTask_ReadProject.EnterProgressFrame();
 		// Always add the mods from the loose directory without using manifests, because they're not packaged together.
 		ReadPluginsInDirectory(FPaths::ProjectModsDir(), EPluginType::Mod, Plugins, ChildPlugins);
+		if (OutPluginSources)
+		{
+			OptionalOutPluginRoots.Add(FPaths::ProjectModsDir());
+		}
 
 		// If they have a list of additional directories to check, add those plugins too
 		for (const FString& Dir : Project->GetAdditionalPluginDirectories())
 		{
 			SlowTask_ReadProject.EnterProgressFrame();
+			if (OutPluginSources)
+			{
+				OptionalOutPluginRoots.Add(Dir);
+			}
 			ReadPluginsInDirectory(Dir, EPluginType::External, Plugins, ChildPlugins);
 		}
 
@@ -884,6 +909,10 @@ void FPluginManager::ReadAllPlugins(FDiscoveredPluginMap& Plugins, const TSet<FS
 		// Add plugins from FPaths::EnterprisePluginsDir if it exists
 		if (FPaths::DirectoryExists(FPaths::EnterprisePluginsDir()))
 		{
+			if (OutPluginSources)
+			{
+				OptionalOutPluginRoots.Add(FPaths::EnterprisePluginsDir());
+			}
 			ReadPluginsInDirectory(FPaths::EnterprisePluginsDir(), EPluginType::Enterprise, Plugins, ChildPlugins);
 		}
 	}
@@ -891,6 +920,10 @@ void FPluginManager::ReadAllPlugins(FDiscoveredPluginMap& Plugins, const TSet<FS
 	for (const FString& ExtraSearchPath : ExtraSearchPaths)
 	{
 		SlowTask_ReadAll.EnterProgressFrame();
+		if (OutPluginSources)
+		{
+			OptionalOutPluginRoots.Add(ExtraSearchPath);
+		}
 		ReadPluginsInDirectory(ExtraSearchPath, EPluginType::External, Plugins, ChildPlugins);
 	}
 
@@ -960,6 +993,27 @@ void FPluginManager::ReadAllPlugins(FDiscoveredPluginMap& Plugins, const TSet<FS
 		}
 	}
 
+	if (OutPluginSources)
+	{
+		for (const FString& Root : OptionalOutPluginRoots)
+		{
+			OutPluginSources->Add(FString::Printf(TEXT("PluginsDir: %s"),
+				*FPaths::ConvertRelativePathToFull(Root)));
+		}
+
+		// Also report PakFiles, since ReadPluginsInDirectory searches PakFiles in addition to directories on disk
+		if (Project)
+		{
+			bool bRunningWithPakFile =
+				(FPlatformFileManager::Get().FindPlatformFile(TEXT("PakFile")) != nullptr);
+			if (bRunningWithPakFile)
+			{
+				// TODO: Are there are pak files that might contain .uplugin files?
+				OutPluginSources->Add(FString::Printf(TEXT("PakFiles: %s"),
+					*FPaths::ConvertRelativePathToFull(FPaths::Combine(FPaths::ProjectContentDir(), TEXT("Paks")))));
+			}
+		}
+	}
 #endif
 }
 
@@ -1426,7 +1480,8 @@ bool FPluginManager::ConfigureEnabledPlugins()
 				{
 					if (!ConfiguredPluginNames.Contains(EnablePluginName) && !ExceptPlugins.Contains(EnablePluginName) && (!bExceptRestrictedPlugins || !IsRestrictedPlugin(EnablePluginName)))
 					{
-						if (!ConfigureEnabledPluginForCurrentTarget(FPluginReferenceDescriptor(EnablePluginName, true), EnabledPlugins))
+						if (!ConfigureEnabledPluginForCurrentTarget(FPluginReferenceDescriptor(EnablePluginName, true), EnabledPlugins,
+							bAllPluginsEnabledViaCommandLine ? TEXTVIEW("Commandline EnableAllPlugins") : TEXTVIEW("CommandLine EnablePlugins=")))
 						{
 							if (bAllPluginsEnabledViaCommandLine)
 							{
@@ -1447,7 +1502,8 @@ bool FPluginManager::ConfigureEnabledPlugins()
 			{
 				if (!ConfiguredPluginNames.Contains(DisablePluginName))
 				{
-					if (!ConfigureEnabledPluginForCurrentTarget(FPluginReferenceDescriptor(DisablePluginName, false), EnabledPlugins))
+					if (!ConfigureEnabledPluginForCurrentTarget(FPluginReferenceDescriptor(DisablePluginName, false), EnabledPlugins,
+						TEXTVIEW("CommandLine DisablePlugins=")))
 					{
 						return false;
 					}
@@ -1471,6 +1527,8 @@ bool FPluginManager::ConfigureEnabledPlugins()
 				const FString ReceiptWildcard = FTargetReceipt::GetDefaultPath(BaseDir, TEXT("*"), FPlatformProcess::GetBinariesSubdirectory(), FApp::GetBuildConfiguration(), nullptr);
 				const FString ReceiptPath = FPaths::GetPath(ReceiptWildcard);
 				IFileManager::Get().FindFiles(AllTargetFilesWithoutPath, *ReceiptWildcard, true, false);
+				FString SourceDescription = FString::Printf(TEXT("Receipt files %s"), *ReceiptWildcard);
+
 				for (const FString& TargetFileWithoutPath : AllTargetFilesWithoutPath)
 				{
 					const FString TargetFile = FPaths::Combine(ReceiptPath, TargetFileWithoutPath);
@@ -1489,7 +1547,8 @@ bool FPluginManager::ConfigureEnabledPlugins()
 									const bool bEnabled = Pair.Value;
 									if (!ConfiguredPluginNames.Contains(PluginName))
 									{
-										if (!ConfigureEnabledPluginForCurrentTarget(FPluginReferenceDescriptor(PluginName, bEnabled), EnabledPlugins))
+										if (!ConfigureEnabledPluginForCurrentTarget(FPluginReferenceDescriptor(PluginName, bEnabled), EnabledPlugins,
+											SourceDescription))
 										{
 											bOutError = true;
 											break;
@@ -1529,7 +1588,8 @@ bool FPluginManager::ConfigureEnabledPlugins()
 				{
 					if (!ConfiguredPluginNames.Contains(TargetEnabledPlugin))
 					{
-						if (!ConfigureEnabledPluginForCurrentTarget(FPluginReferenceDescriptor(TargetEnabledPlugin, true), EnabledPlugins))
+						if (!ConfigureEnabledPluginForCurrentTarget(FPluginReferenceDescriptor(TargetEnabledPlugin, true), EnabledPlugins,
+							TEXTVIEW("UBT_TARGET_ENABLED_PLUGINS")))
 						{
 							return false;
 						}
@@ -1546,7 +1606,8 @@ bool FPluginManager::ConfigureEnabledPlugins()
 				{
 					if (!ConfiguredPluginNames.Contains(TargetDisabledPlugin))
 					{
-						if (!ConfigureEnabledPluginForCurrentTarget(FPluginReferenceDescriptor(TargetDisabledPlugin, false), EnabledPlugins))
+						if (!ConfigureEnabledPluginForCurrentTarget(FPluginReferenceDescriptor(TargetDisabledPlugin, false), EnabledPlugins,
+							TEXTVIEW("UBT_TARGET_ENABLED_PLUGINS")))
 						{
 							return false;
 						}
@@ -1556,13 +1617,15 @@ bool FPluginManager::ConfigureEnabledPlugins()
 			}
 #endif // READ_TARGET_ENABLED_PLUGINS_FROM_RECEIPT
 
-			auto ProcessPluginConfigurations = [&ConfiguredPluginNames, &EnabledPlugins, this](const TArray<FPluginReferenceDescriptor>& PluginReferences)->bool
+			auto ProcessPluginConfigurations =
+				[&ConfiguredPluginNames, &EnabledPlugins, this]
+				(const TArray<FPluginReferenceDescriptor>& PluginReferences, FStringView SourceDescription)->bool
 			{
 				for (const FPluginReferenceDescriptor& PluginReference : PluginReferences)
 				{
 					if (!ConfiguredPluginNames.Contains(PluginReference.Name))
 					{
-						if (!ConfigureEnabledPluginForCurrentTarget(PluginReference, EnabledPlugins))
+						if (!ConfigureEnabledPluginForCurrentTarget(PluginReference, EnabledPlugins, SourceDescription))
 						{
 							return false;
 						}
@@ -1584,10 +1647,11 @@ bool FPluginManager::ConfigureEnabledPlugins()
 
 					// Copy the plugin references, since we may modify the project if any plugins are missing
 					TArray<FPluginReferenceDescriptor> PluginReferences(ProjectDescriptor->Plugins);
-					if (!ProcessPluginConfigurations(PluginReferences))
-							{
-								return false;
-							}
+					if (!ProcessPluginConfigurations(PluginReferences,
+						FString::Printf(TEXT("Enabled plugins in .uproject for %s"), *ProjectDescriptor->Description)))
+					{
+						return false;
+					}
 				}
 			}
 
@@ -1595,12 +1659,14 @@ bool FPluginManager::ConfigureEnabledPlugins()
 				// Add the plugins which are enabled by default
 				SCOPED_BOOT_TIMING("AddPluginsEnabledByDefault");
 
+				FString SourceDescription = FString::Printf(TEXT(".uplugin files that are enabled by default"));
 				for (const FString& PluginName : PluginsToConfigure)
 				{
 					const TSharedRef<FPlugin>& Plugin = DiscoveredPluginMapUtils::FindPluginInMap_Checked(AllPlugins, PluginName);
 					if (Plugin->IsEnabledByDefault(bAllowEnginePluginsEnabledByDefault) && !ConfiguredPluginNames.Contains(PluginName))
 					{
-						if (!ConfigureEnabledPluginForCurrentTarget(FPluginReferenceDescriptor(PluginName, true), EnabledPlugins))
+						if (!ConfigureEnabledPluginForCurrentTarget(FPluginReferenceDescriptor(PluginName, true),
+							EnabledPlugins, SourceDescription))
 						{
 							return false;
 						}
@@ -1615,6 +1681,7 @@ bool FPluginManager::ConfigureEnabledPlugins()
 		{
 			TArray<FString> ProgramPluginNames;
 			GConfig->GetArray(TEXT("Plugins"), ConfigEntry, ProgramPluginNames, GEngineIni);
+			FString SourceDescription = FString::Printf(TEXT("Engine.ini:[%s]:ProgramEnabledPlugins"), ConfigEntry);
 
 			for (const FString& PluginName : ProgramPluginNames)
 			{
@@ -1622,7 +1689,8 @@ bool FPluginManager::ConfigureEnabledPlugins()
 				{
 					FPluginReferenceDescriptor PluginReference(PluginName, true);
 					PluginReference.bOptional = bOptionalPlugin;
-					if (!ConfigureEnabledPluginForCurrentTarget(PluginReference, EnabledPlugins))
+					if (!ConfigureEnabledPluginForCurrentTarget(PluginReference,
+						EnabledPlugins, SourceDescription))
 					{
 						return false;
 					}
@@ -2037,7 +2105,9 @@ bool FPluginManager::RequiresTempTargetForCodePlugin(const FProjectDescriptor* P
 	return false;
 }
 
-bool FPluginManager::GetCodePluginsForProject(const FProjectDescriptor* ProjectDescriptor, const FString& Platform, EBuildConfiguration Configuration, EBuildTargetType TargetType, FDiscoveredPluginMap& AllPlugins, TSet<FString>& CodePluginNames, FConfigurePluginResultInfo& OutResultInfo)
+bool FPluginManager::GetCodePluginsForProject(const FProjectDescriptor* ProjectDescriptor, const FString& Platform,
+	EBuildConfiguration Configuration, EBuildTargetType TargetType, FDiscoveredPluginMap& AllPlugins,
+	TSet<FString>& CodePluginNames, FConfigurePluginResultInfo& OutResultInfo)
 {
 	// Can only check the current project at the moment, since we won't have enumerated them otherwise
 	check(ProjectDescriptor == nullptr || ProjectDescriptor == IProjectManager::Get().GetCurrentProject());
@@ -2063,7 +2133,9 @@ bool FPluginManager::GetCodePluginsForProject(const FProjectDescriptor* ProjectD
 		{
 			if(!ConfiguredPluginNames.Contains(PluginReference.Name))
 			{
-				if (!ConfigureEnabledPluginForTarget(PluginReference, ProjectDescriptor, FString(), Platform, Configuration, TargetType, bLoadPluginsForTargetPlatforms, AllPlugins, EnabledPlugins, OutResultInfo))
+				if (!ConfigureEnabledPluginForTarget(PluginReference, ProjectDescriptor, FString(), Platform,
+					Configuration, TargetType, bLoadPluginsForTargetPlatforms, AllPlugins, EnabledPlugins,
+					OutResultInfo))
 				{
 					return false;
 				}
@@ -2077,7 +2149,9 @@ bool FPluginManager::GetCodePluginsForProject(const FProjectDescriptor* ProjectD
 	{
 		if (DiscoveredPluginMapUtils::ResolvePluginFromMapVal(PluginPair.Value)->IsEnabledByDefault(bAllowEnginePluginsEnabledByDefault) && !ConfiguredPluginNames.Contains(PluginPair.Key))
 		{
-			if (!ConfigureEnabledPluginForTarget(FPluginReferenceDescriptor(PluginPair.Key, true), ProjectDescriptor, FString(), Platform, Configuration, TargetType, bLoadPluginsForTargetPlatforms, AllPlugins, EnabledPlugins, OutResultInfo))
+			if (!ConfigureEnabledPluginForTarget(FPluginReferenceDescriptor(PluginPair.Key, true), ProjectDescriptor,
+				FString(), Platform, Configuration, TargetType, bLoadPluginsForTargetPlatforms, AllPlugins,
+				EnabledPlugins, OutResultInfo))
 			{
 				return false;
 			}
@@ -2103,38 +2177,24 @@ bool FPluginManager::GetCodePluginsForProject(const FProjectDescriptor* ProjectD
 	return true;
 }
 
-bool FPluginManager::ConfigureEnabledPluginForCurrentTarget(const FPluginReferenceDescriptor& FirstReference, TMap<FString, FPlugin*>& EnabledPlugins)
+bool FPluginManager::ConfigureEnabledPluginForCurrentTarget(const FPluginReferenceDescriptor& FirstReference,
+	TMap<FString, FPlugin*>& EnabledPlugins, FStringView SourceOfPluginRequest)
 {
 	SCOPED_BOOT_TIMING("ConfigureEnabledPluginForCurrentTarget");
 
 	FConfigurePluginResultInfo ResultInfo;
 
-	if (!ConfigureEnabledPluginForTarget(FirstReference, IProjectManager::Get().GetCurrentProject(), UE_APP_NAME, FPlatformMisc::GetUBTPlatform(), FApp::GetBuildConfiguration(), FApp::GetBuildTargetType(), (bool)LOAD_PLUGINS_FOR_TARGET_PLATFORMS, AllPlugins, EnabledPlugins, ResultInfo))
+	if (ConfigureEnabledPluginForTarget(FirstReference, IProjectManager::Get().GetCurrentProject(), UE_APP_NAME,
+		FPlatformMisc::GetUBTPlatform(), FApp::GetBuildConfiguration(), FApp::GetBuildTargetType(),
+		(bool)LOAD_PLUGINS_FOR_TARGET_PLATFORMS, AllPlugins, EnabledPlugins, ResultInfo))
 	{
-		// If we're in unattended mode, don't open any windows and fatal out
-		if (FApp::IsUnattended())
-		{
-			switch (ResultInfo.Code)
-			{
-			case EConfigurePluginResultCode::PluginMissing:
-				UE_LOG(LogPluginManager, Fatal, TEXT("This project requires the '%s' plugin. Install it and try again, or remove it from the project's required plugin list."), *ResultInfo.PluginReference->Name);
-				break;
-
-			case EConfigurePluginResultCode::PluginSealed:
-				UE_LOG(LogPluginManager, Fatal, TEXT("This project requires the '%s' plugin, which has an illegal dependency on the sealed plugin '%s'."), *FirstReference.Name, *ResultInfo.PluginReference->Name);
-				break;
-
-			case EConfigurePluginResultCode::DisallowedDependency:
-				UE_LOG(LogPluginManager, Fatal, TEXT("This project requires the '%s' plugin, which has a disallowed dependency on the plugin '%s'."), *FirstReference.Name, *ResultInfo.PluginReference->Name);
-				break;
-
-			default:
-				checkf(0, TEXT("Unimplemented EConfigurePluginResultCode switch case"));
-			}
-			return false;
-		}
+		return true;
+	}
 
 #if !IS_MONOLITHIC
+	// If we're in unattended mode, don't open any windows
+	if (!FApp::IsUnattended())
+	{
 		switch (ResultInfo.Code)
 		{
 		case EConfigurePluginResultCode::PluginMissing:
@@ -2180,189 +2240,321 @@ bool FPluginManager::ConfigureEnabledPluginForCurrentTarget(const FPluginReferen
 		default:
 			checkf(0, TEXT("Unimplemented EConfigurePluginResultCode switch case"));
 		}
+	}
 #endif
 
-		// Unable to continue
-		UE_LOG(LogPluginManager, Error, TEXT("Unable to load plugin '%s'. Aborting."), *ResultInfo.PluginReference->Name);
-		return false;
+	// Unable to continue
+	// If we're in unattended mode, set error to fatal, otherwise let the caller decide whether to kill the process.
+	// Log more diagnostics in non-ship; don't give them in ship because they include full paths on the user's machine.
+	bool bFatalError = FApp::IsUnattended();
+	FString ErrorMessage;
+
+#if UE_BUILD_SHIPPING
+	ErrorMessage = FString::Printf(TEXT("Unable to load plugin '%s'. Aborting."), *ResultInfo.PluginReference->Name);
+#else
+	auto ReferenceChainToString = [](TArray<const FPluginReferenceDescriptor*>& Chain, const FPluginReferenceDescriptor* End)
+		{
+			TStringBuilder<256> ChainStr;
+			for (const FPluginReferenceDescriptor* Iter : Chain)
+			{
+				ChainStr << Iter->Name << TEXTVIEW(" -> ");
+			}
+			ChainStr << End->Name;
+			return FString(ChainStr);
+		};
+	TArray<FString> PluginLocations;
+	FDiscoveredPluginMap UnusedListOfDiscoveredPlugins;
+	ReadAllPlugins(UnusedListOfDiscoveredPlugins, PluginDiscoveryPaths, &PluginLocations);
+	TStringBuilder<1024> PluginLocationsStr;
+	for (const FString& Location : PluginLocations)
+	{
+		PluginLocationsStr << TEXT("\n\t") << Location;
 	}
-	return true;
+	if (ResultInfo.PluginReference == &FirstReference)
+	{
+		switch (ResultInfo.Code)
+		{
+		case EConfigurePluginResultCode::PluginMissing:
+			ErrorMessage = FString::Printf(
+				TEXT("Unable to load plugin '%s'. It was requested by %.*s, but is missing on disk. Aborting. Looked in these locations for .uplugin files:")
+				TEXT("%s"),
+				*FirstReference.Name, SourceOfPluginRequest.Len(), SourceOfPluginRequest.GetData(),
+				*PluginLocationsStr);
+			break;
+		case EConfigurePluginResultCode::PluginSealed:
+			ErrorMessage = FString::Printf(
+				TEXT("Unable to load plugin '%s'. It was requested by %.*s, but is is a sealed plugin and is therefore not usable. Aborting."),
+				*FirstReference.Name, SourceOfPluginRequest.Len(), SourceOfPluginRequest.GetData());
+			break;
+		case EConfigurePluginResultCode::DisallowedDependency:
+			ErrorMessage = FString::Printf(
+				TEXT("Unable to load plugin '%s'. It was requested by %.*s, but is a disallowed plugin is therefore not usable. Aborting."),
+				*FirstReference.Name, SourceOfPluginRequest.Len(), SourceOfPluginRequest.GetData());
+			break;
+		default:
+			checkf(0, TEXT("Unimplemented EConfigurePluginResultCode switch case"));
+		}
+	}
+	else
+	{
+		switch (ResultInfo.Code)
+		{
+		case EConfigurePluginResultCode::PluginMissing:
+			ErrorMessage = FString::Printf(
+				TEXT("Unable to load plugin '%s'. It was requested by %.*s, but it has a dependency that is missing on disk.")
+				TEXT("\n\tThe missing dependency is %s (ReferenceChain: %s).")
+				TEXT("\n\tAborting.")
+				TEXT("\n\tLooked in these locations for .uplugin files:")
+				TEXT("%s"),
+				*FirstReference.Name, SourceOfPluginRequest.Len(), SourceOfPluginRequest.GetData(),
+				*ResultInfo.PluginReference->Name, *ReferenceChainToString(ResultInfo.ReferenceChain, ResultInfo.PluginReference),
+				*PluginLocationsStr);
+			break;
+		case EConfigurePluginResultCode::PluginSealed:
+			ErrorMessage = FString::Printf(
+				TEXT("Unable to load plugin '%s'. It was requested by %.*s, but it has a dependency on a sealed plugin.")
+				TEXT("\n\tThe sealed dependency is %s (ReferenceChain: %s).")
+				TEXT("\n\tAborting."),
+				*FirstReference.Name, SourceOfPluginRequest.Len(), SourceOfPluginRequest.GetData(),
+				*ResultInfo.PluginReference->Name, *ReferenceChainToString(ResultInfo.ReferenceChain, ResultInfo.PluginReference));
+			break;
+		case EConfigurePluginResultCode::DisallowedDependency:
+			ErrorMessage = FString::Printf(
+				TEXT("Unable to load plugin '%s'. It was requested by %.*s, but it has a dependency on a disallowed plugin.")
+				TEXT("\n\tThe disallowed plugin is %s (ReferenceChain: %s).")
+				TEXT("\n\tAborting."),
+				*FirstReference.Name, SourceOfPluginRequest.Len(), SourceOfPluginRequest.GetData(),
+				*ResultInfo.PluginReference->Name, *ReferenceChainToString(ResultInfo.ReferenceChain, ResultInfo.PluginReference));
+			break;
+		default:
+			checkf(0, TEXT("Unimplemented EConfigurePluginResultCode switch case"));
+		}
+	}
+#endif
+
+	if (bFatalError)
+	{
+		UE_LOG(LogPluginManager, Fatal, TEXT("%s"), *ErrorMessage);
+	}
+	else
+	{
+		UE_LOG(LogPluginManager, Error, TEXT("%s"), *ErrorMessage);
+	}
+	return false;
 }
 
-bool FPluginManager::ConfigureEnabledPluginForTarget(const FPluginReferenceDescriptor& FirstReference, const FProjectDescriptor* ProjectDescriptor, const FString& TargetName, const FString& Platform, EBuildConfiguration Configuration, EBuildTargetType TargetType, bool bLoadPluginsForTargetPlatforms, FDiscoveredPluginMap& AllPlugins, TMap<FString, FPlugin*>& EnabledPlugins, FConfigurePluginResultInfo& OutResultInfo)
+bool FPluginManager::ConfigureEnabledPluginForTarget(const FPluginReferenceDescriptor& FirstReference,
+	const FProjectDescriptor* ProjectDescriptor, const FString& TargetName, const FString& Platform,
+	EBuildConfiguration Configuration, EBuildTargetType TargetType, bool bLoadPluginsForTargetPlatforms,
+	FDiscoveredPluginMap& AllPlugins, TMap<FString, FPlugin*>& EnabledPlugins,
+	FConfigurePluginResultInfo& OutResultInfo)
 {
-	if (!EnabledPlugins.Contains(FirstReference.Name))
+	if (EnabledPlugins.Contains(FirstReference.Name))
 	{
-		// Set of plugin names we've added to the queue for processing
-		TSet<FString> NewPluginNames;
-		NewPluginNames.Add(FirstReference.Name);
+		// Already enabled. Just verify the version
+		if (FirstReference.RequestedVersion.IsSet())
+		{
+			UE_CLOG(
+				(*EnabledPlugins.Find(FirstReference.Name))->GetDescriptor().Version != FirstReference.RequestedVersion.GetValue(),
+				LogPluginManager, Error,
+				TEXT("Requested explicit version (v%d) of plugin '%s', but a different version (v%d) was already enabled by another source."),
+				FirstReference.RequestedVersion.GetValue(), *FirstReference.Name,
+				(*EnabledPlugins.Find(FirstReference.Name))->GetDescriptor().Version
+			);
+		}
+		return true;
+	}
 
-		// Queue of plugin references to consider
-		TArray<const FPluginReferenceDescriptor*> NewPluginReferences;
-		NewPluginReferences.Add(&FirstReference);
+	// Set of plugin names we've added to the queue for processing, and the plugin that referenced
+	// them, so we can display the chain of references if necessary for diagnostics
+	TMap<FString, const FPluginReferenceDescriptor*> SeenPlugins;
 
-		// Loop through the queue of plugin references that need to be enabled, queuing more items as we go
-		TArray<TSharedRef<FPlugin>> NewPlugins;
-		for (int32 Idx = 0; Idx < NewPluginReferences.Num(); Idx++)
-		{
-			const FPluginReferenceDescriptor& Reference = *NewPluginReferences[Idx];
+	// Queue of plugin references to consider
+	TArray<const FPluginReferenceDescriptor*> NewPluginQueue;
 
-			// Check if the plugin is required for this platform
-			if(!Reference.IsEnabledForPlatform(Platform) || !Reference.IsEnabledForTargetConfiguration(Configuration) || !Reference.IsEnabledForTarget(TargetType))
+	OutResultInfo.ReferenceChain.Reset();
+	auto SetReferenceChain = [&SeenPlugins, &OutResultInfo](const FPluginReferenceDescriptor* Reference)
+		{
+			while (Reference)
 			{
-				UE_LOG(LogPluginManager, Verbose, TEXT("Ignoring plugin '%s' for platform/configuration"), *Reference.Name);
-				continue;
+				const FPluginReferenceDescriptor** Next = SeenPlugins.Find(Reference->Name);
+				// SeenPlugins[FIrstReference] == FirstReference, so stop when we reach that cycle. No other cycles are possible.
+				if (!Next || Reference == *Next) 
+				{
+					break;
+				}
+				Reference = *Next;
+				// The reference chain should not include the input Reference, so add at the end of loop rather than the beginning
+				OutResultInfo.ReferenceChain.Add(Reference);
 			}
+			Algo::Reverse(OutResultInfo.ReferenceChain);
+		};
 
-			// Check if the plugin is required for this platform
-			if(!bLoadPluginsForTargetPlatforms && !Reference.IsSupportedTargetPlatform(Platform))
-			{
-				UE_LOG(LogPluginManager, Verbose, TEXT("Ignoring plugin '%s' due to unsupported platform"), *Reference.Name);
-				continue;
-			}
 
-			// Find the plugin being enabled
-			const TSharedRef<FPlugin>* PluginPtr = DiscoveredPluginMapUtils::FindPluginInMap(AllPlugins, Reference.Name);
-			if (PluginPtr && Reference.RequestedVersion.IsSet())
-			{
-				TArray<TSharedRef<FPlugin>>* PluginVersions = DiscoveredPluginMapUtils::FindAllPluginVersionsWithName(AllPlugins, Reference.Name);
-				check(PluginVersions);
+	// Loop through the queue of plugin references that need to be enabled, queuing more items as we go
+	NewPluginQueue.Add(&FirstReference);
+	SeenPlugins.Add(FirstReference.Name, &FirstReference);
+	for (int32 Idx = 0; Idx < NewPluginQueue.Num(); Idx++)
+	{
+		const FPluginReferenceDescriptor& Reference = *NewPluginQueue[Idx];
 
-				TSharedRef<FPlugin>* FoundVersion = nullptr;
-				for (TSharedRef<FPlugin>& PluginVersion : *PluginVersions)
-				{
-					if (PluginVersion->GetDescriptor().Version == Reference.RequestedVersion.GetValue())
-					{
-						FoundVersion = &PluginVersion;
-						break;
-					}
-				}
+		// Check if the plugin is required for this platform
+		if(!Reference.IsEnabledForPlatform(Platform) || !Reference.IsEnabledForTargetConfiguration(Configuration) || !Reference.IsEnabledForTarget(TargetType))
+		{
+			UE_LOG(LogPluginManager, Verbose, TEXT("Ignoring plugin '%s' for platform/configuration"), *Reference.Name);
+			continue;
+		}
 
-				if (!FoundVersion)
-				{
-					UE_LOG(LogPluginManager, Warning, TEXT("Failed to find specific version (v%d) of plugin '%s'. Other versions exist, but the explicit version requested was missing."),
-						Reference.RequestedVersion.GetValue(), *Reference.Name);
+		// Check if the plugin is required for this platform
+		if(!bLoadPluginsForTargetPlatforms && !Reference.IsSupportedTargetPlatform(Platform))
+		{
+			UE_LOG(LogPluginManager, Verbose, TEXT("Ignoring plugin '%s' due to unsupported platform"), *Reference.Name);
+			continue;
+		}
 
-					PluginPtr = nullptr;
-				}
-				else if (FoundVersion != PluginPtr && (*PluginPtr)->bEnabled)
-				{
-					UE_LOG(LogPluginManager, Warning, TEXT("A different version of the %s plugin (v%d : '%s') is already enabled. Cannot switch to v%d (%s) while another version is active."),
-						*Reference.Name, (*PluginPtr)->Descriptor.Version, *(*PluginPtr)->FileName, Reference.RequestedVersion.GetValue(), *(*FoundVersion)->FileName);
+		// Find the plugin being enabled
+		const TSharedRef<FPlugin>* PluginPtr = DiscoveredPluginMapUtils::FindPluginInMap(AllPlugins, Reference.Name);
+		if (PluginPtr && Reference.RequestedVersion.IsSet())
+		{
+			TArray<TSharedRef<FPlugin>>* PluginVersions = DiscoveredPluginMapUtils::FindAllPluginVersionsWithName(AllPlugins, Reference.Name);
+			check(PluginVersions);
 
-					PluginPtr = nullptr;
-				}
-				else
+			TSharedRef<FPlugin>* FoundVersion = nullptr;
+			for (TSharedRef<FPlugin>& PluginVersion : *PluginVersions)
+			{
+				if (PluginVersion->GetDescriptor().Version == Reference.RequestedVersion.GetValue())
 				{
-					PluginPtr = DiscoveredPluginMapUtils::PromotePluginToOfferedVersion(AllPlugins, FoundVersion);
+					FoundVersion = &PluginVersion;
+					break;
 				}
 			}
 
-			if (PluginPtr == nullptr)
+			if (!FoundVersion)
 			{
-				// Ignore any optional plugins
-				if (Reference.bOptional)
-				{
-					UE_LOG(LogPluginManager, Verbose, TEXT("Ignored optional reference to '%s' plugin; plugin was not found."), *Reference.Name);
-					continue;
-				}
+				UE_LOG(LogPluginManager, Warning, TEXT("Failed to find specific version (v%d) of plugin '%s'. Other versions exist, but the explicit version requested was missing."),
+					Reference.RequestedVersion.GetValue(), *Reference.Name);
 
-				// Add it to the missing list
-				OutResultInfo = FConfigurePluginResultInfo(EConfigurePluginResultCode::PluginMissing, &Reference);
-				return false;
+				PluginPtr = nullptr;
 			}
-
-			// Check the plugin is not disabled by the platform
-			FPlugin& Plugin = PluginPtr->Get();
-
-			// Allow the platform to disable it
-			if (FPlatformMisc::ShouldDisablePluginAtRuntime(Plugin.Name))
+			else if (FoundVersion != PluginPtr && (*PluginPtr)->bEnabled)
 			{
-				UE_LOG(LogPluginManager, Verbose, TEXT("Plugin '%s' was disabled by platform."), *Reference.Name);
-				continue;
-			}
+				UE_LOG(LogPluginManager, Warning, TEXT("A different version of the %s plugin (v%d : '%s') is already enabled. Cannot switch to v%d (%s) while another version is active."),
+					*Reference.Name, (*PluginPtr)->Descriptor.Version, *(*PluginPtr)->FileName, Reference.RequestedVersion.GetValue(), *(*FoundVersion)->FileName);
 
-			// Check the plugin supports this platform
-			if (!bLoadPluginsForTargetPlatforms && !Plugin.Descriptor.SupportsTargetPlatform(Platform))
-			{
-				UE_LOG(LogPluginManager, Verbose, TEXT("Ignoring plugin '%s' due to unsupported platform in plugin descriptor"), *Reference.Name);
-				continue;
+				PluginPtr = nullptr;
 			}
-
-			// Check that this plugin supports the current program
-			if (TargetType == EBuildTargetType::Program && !Plugin.Descriptor.SupportedPrograms.Contains(TargetName))
+			else
 			{
-				UE_LOG(LogPluginManager, Verbose, TEXT("Ignoring plugin '%s' due to absence from the supported programs list"), *Reference.Name);
-				continue;
+				PluginPtr = DiscoveredPluginMapUtils::PromotePluginToOfferedVersion(AllPlugins, FoundVersion);
 			}
+		}
 
-			// Skip loading Enterprise plugins when project is not an Enterprise project
-			if (Plugin.Type == EPluginType::Enterprise && ProjectDescriptor != nullptr && !ProjectDescriptor->bIsEnterpriseProject)
+		if (PluginPtr == nullptr)
+		{
+			// Ignore any optional plugins
+			if (Reference.bOptional)
 			{
-				UE_LOG(LogPluginManager, Verbose, TEXT("Ignoring plugin '%s' due to not being an Enterpise project"), *Reference.Name);
+				UE_LOG(LogPluginManager, Verbose, TEXT("Ignored optional reference to '%s' plugin; plugin was not found."), *Reference.Name);
 				continue;
 			}
 
+			// Add it to the missing list
+			OutResultInfo = FConfigurePluginResultInfo(EConfigurePluginResultCode::PluginMissing, &Reference);
+			SetReferenceChain(&Reference);
+			return false;
+		}
+
+		// Check the plugin is not disabled by the platform
+		FPlugin& Plugin = PluginPtr->Get();
+
+		// Allow the platform to disable it
+		if (FPlatformMisc::ShouldDisablePluginAtRuntime(Plugin.Name))
+		{
+			UE_LOG(LogPluginManager, Verbose, TEXT("Plugin '%s' was disabled by platform."), *Reference.Name);
+			continue;
+		}
+
+		// Check the plugin supports this platform
+		if (!bLoadPluginsForTargetPlatforms && !Plugin.Descriptor.SupportsTargetPlatform(Platform))
+		{
+			UE_LOG(LogPluginManager, Verbose, TEXT("Ignoring plugin '%s' due to unsupported platform in plugin descriptor"), *Reference.Name);
+			continue;
+		}
+
+		// Check that this plugin supports the current program
+		if (TargetType == EBuildTargetType::Program && !Plugin.Descriptor.SupportedPrograms.Contains(TargetName))
+		{
+			UE_LOG(LogPluginManager, Verbose, TEXT("Ignoring plugin '%s' due to absence from the supported programs list"), *Reference.Name);
+			continue;
+		}
+
+		// Skip loading Enterprise plugins when project is not an Enterprise project
+		if (Plugin.Type == EPluginType::Enterprise && ProjectDescriptor != nullptr && !ProjectDescriptor->bIsEnterpriseProject)
+		{
+			UE_LOG(LogPluginManager, Verbose, TEXT("Ignoring plugin '%s' due to not being an Enterpise project"), *Reference.Name);
+			continue;
+		}
+
 #if IS_MONOLITHIC
-			// In monolithic builds, if the plugin is optional, and any modules are not compiled in this build, do not load it. This makes content compiled with plugins enabled compatible with builds where they are disabled.
-			if (Reference.bOptional)
+		// In monolithic builds, if the plugin is optional, and any modules are not compiled in this build, do not load it. This makes content compiled with plugins enabled compatible with builds where they are disabled.
+		if (Reference.bOptional)
+		{
+			bool bAllModulesAvailable = true;
+			for (const FModuleDescriptor& ModuleDescriptor : Plugin.Descriptor.Modules)
 			{
-				bool bAllModulesAvailable = true;
-				for (const FModuleDescriptor& ModuleDescriptor : Plugin.Descriptor.Modules)
+				if (ModuleDescriptor.IsLoadedInCurrentConfiguration() && !FModuleManager::Get().ModuleExists(*ModuleDescriptor.Name.ToString()))
 				{
-					if (ModuleDescriptor.IsLoadedInCurrentConfiguration() && !FModuleManager::Get().ModuleExists(*ModuleDescriptor.Name.ToString()))
-					{
-						bAllModulesAvailable = false;
-						break;
-					}
+					bAllModulesAvailable = false;
+					break;
 				}
+			}
 
-				if (!bAllModulesAvailable)
-				{
-					UE_LOG(LogPluginManager, Verbose, TEXT("Ignored optional reference to '%s' plugin; plugin's modules were not found."), *Reference.Name);
-					continue;
-				}
+			if (!bAllModulesAvailable)
+			{
+				UE_LOG(LogPluginManager, Verbose, TEXT("Ignored optional reference to '%s' plugin; plugin's modules were not found."), *Reference.Name);
+				continue;
 			}
+		}
 #endif // IS_MONOLITHIC
 
-			// Add references to all its dependencies
-			for (const FPluginReferenceDescriptor& NextReference : Plugin.Descriptor.Plugins)
+		// Add references to all its dependencies
+		for (const FPluginReferenceDescriptor& NextReference : Plugin.Descriptor.Plugins)
+		{
+			const TSharedRef<FPlugin>* NextPluginPtr = DiscoveredPluginMapUtils::FindPluginInMap(AllPlugins, NextReference.Name);
+			bool bIsNewlySeen = false;
+			const FPluginReferenceDescriptor*& SeenReferencer = SeenPlugins.FindOrAdd(NextReference.Name, nullptr);
+			if (!SeenReferencer)
 			{
-				const TSharedRef<FPlugin>* NextPluginPtr = DiscoveredPluginMapUtils::FindPluginInMap(AllPlugins, NextReference.Name);
-				if (NextPluginPtr != nullptr && NextPluginPtr->Get().Descriptor.bIsSealed)
-				{
-					OutResultInfo = FConfigurePluginResultInfo(EConfigurePluginResultCode::PluginSealed, &NextReference);
-					return false;
-				}
+				bIsNewlySeen = true;
+				SeenReferencer = &Reference;
+			}
 
-				if (NextPluginPtr != nullptr && Plugin.Descriptor.DisallowedPlugins.ContainsByPredicate([&NextPluginPtr](const FPluginDisallowedDescriptor& Other) { return Other.Name == (*NextPluginPtr)->Name; }))
-				{
-					OutResultInfo = FConfigurePluginResultInfo(EConfigurePluginResultCode::DisallowedDependency, &NextReference);
-					return false;
-				}
+			if (NextPluginPtr != nullptr && NextPluginPtr->Get().Descriptor.bIsSealed)
+			{
+				OutResultInfo = FConfigurePluginResultInfo(EConfigurePluginResultCode::PluginSealed, &NextReference);
+				SetReferenceChain(&NextReference);
+				return false;
+			}
 
-				if (!EnabledPlugins.Contains(NextReference.Name) && !NewPluginNames.Contains(NextReference.Name))
-				{
-					NewPluginNames.Add(NextReference.Name);
-					NewPluginReferences.Add(&NextReference);
-				}
+			if (NextPluginPtr != nullptr && Plugin.Descriptor.DisallowedPlugins.ContainsByPredicate(
+				[&NextPluginPtr](const FPluginDisallowedDescriptor& Other) { return Other.Name == (*NextPluginPtr)->Name; }))
+			{
+				OutResultInfo = FConfigurePluginResultInfo(EConfigurePluginResultCode::DisallowedDependency, &NextReference);
+				// If the dependency was previously referenced by another plugin, overwrite that other plugin with the current Plugin,
+				// so that the caller will see the reference chain that has the DisallowedDependencye
+				SeenReferencer = &Reference;
+				SetReferenceChain(&NextReference);
+				return false;
 			}
 
-			// Add the plugin
-			EnabledPlugins.Add(Plugin.GetName(), &Plugin);
+			if (!EnabledPlugins.Contains(NextReference.Name) && bIsNewlySeen)
+			{
+				NewPluginQueue.Add(&NextReference);
+			}
 		}
-	}
-	else if (FirstReference.RequestedVersion.IsSet())
-	{
-		UE_CLOG(
-			(*EnabledPlugins.Find(FirstReference.Name))->GetDescriptor().Version != FirstReference.RequestedVersion.GetValue(),
-			LogPluginManager, 
-			Error, 
-			TEXT("Requested explicit version (v%d) of plugin '%s', but a different version (v%d) was already enabled by another source."),
-			FirstReference.RequestedVersion.GetValue(),
-			*FirstReference.Name,
-			(*EnabledPlugins.Find(FirstReference.Name))->GetDescriptor().Version
-		);
+
+		// Add the plugin
+		EnabledPlugins.Add(Plugin.GetName(), &Plugin);
 	}
 	return true;
 }
diff --git a/Engine/Source/Runtime/Projects/Private/PluginManager.h b/Engine/Source/Runtime/Projects/Private/PluginManager.h
--- a/Engine/Source/Runtime/Projects/Private/PluginManager.h
+++ b/Engine/Source/Runtime/Projects/Private/PluginManager.h
@@ -206,7 +206,8 @@ private:
 	void DiscoverAllPlugins();
 
 	/** Reads all the plugin descriptors */
-	static void ReadAllPlugins(FDiscoveredPluginMap& Plugins, const TSet<FString>& ExtraSearchPaths);
+	static void ReadAllPlugins(FDiscoveredPluginMap& Plugins, const TSet<FString>& ExtraSearchPaths,
+		TArray<FString>* OutPluginSources=nullptr);
 
 	/** Reads all the plugin descriptors from disk */
 	static void ReadPluginsInDirectory(const FString& PluginsDirectory, const EPluginType Type, FDiscoveredPluginMap& Plugins, TArray<TSharedRef<FPlugin>>& ChildPlugins);
@@ -227,7 +228,8 @@ private:
 	bool ConfigureEnabledPlugins();
 
 	/** Adds a single enabled plugin, and all its dependencies */
-	bool ConfigureEnabledPluginForCurrentTarget(const FPluginReferenceDescriptor& FirstReference, TMap<FString, FPlugin*>& EnabledPlugins);
+	bool ConfigureEnabledPluginForCurrentTarget(const FPluginReferenceDescriptor& FirstReference,
+		TMap<FString, FPlugin*>& EnabledPlugins, FStringView SourceOfPluginRequest);
 
 	/** Adds a single enabled plugin and all its dependencies. */
 	static bool ConfigureEnabledPluginForTarget(const FPluginReferenceDescriptor& FirstReference, const FProjectDescriptor* ProjectDescriptor, const FString& TargetName, const FString& Platform, EBuildConfiguration Configuration, EBuildTargetType TargetType, bool bLoadPluginsForTargetPlatforms, FDiscoveredPluginMap& AllPlugins, TMap<FString, FPlugin*>& EnabledPlugins, FConfigurePluginResultInfo& OutResultInfo);
