Added option to prevent unload of code library for GameFeaturePlugins.

diff --git a/Engine/Plugins/Experimental/GameFeatures/Source/GameFeatures/Private/GameFeaturePluginStateMachine.cpp b/Engine/Plugins/Experimental/GameFeatures/Source/GameFeatures/Private/GameFeaturePluginStateMachine.cpp
--- a/Engine/Plugins/Experimental/GameFeatures/Source/GameFeatures/Private/GameFeaturePluginStateMachine.cpp
+++ b/Engine/Plugins/Experimental/GameFeatures/Source/GameFeatures/Private/GameFeaturePluginStateMachine.cpp
@@ -1637,7 +1637,7 @@ struct FGameFeaturePluginState_Unmounting : public FGameFeaturePluginState
 			// This will also cause any assets we added to the asset manager to be removed.
 			// Scan paths added to the asset manager should have already been cleaned up.
 			FText FailureReason;
-			if (!IPluginManager::Get().UnmountExplicitlyLoadedPlugin(StateProperties.PluginName, &FailureReason))
+			if (!IPluginManager::Get().UnmountExplicitlyLoadedPlugin(StateProperties.PluginName, &FailureReason, false))
 			{
 				const FStringView ShortUrl = StateProperties.PluginIdentifier.GetIdentifyingString();
 				ensureMsgf(false, TEXT("Failed to explicitly unmount the PluginURL(%.*s) because %s"), ShortUrl.Len(), ShortUrl.GetData(), *FailureReason.ToString());
diff --git a/Engine/Source/Runtime/Core/Private/Modules/ModuleManager.cpp b/Engine/Source/Runtime/Core/Private/Modules/ModuleManager.cpp
--- a/Engine/Source/Runtime/Core/Private/Modules/ModuleManager.cpp
+++ b/Engine/Source/Runtime/Core/Private/Modules/ModuleManager.cpp
@@ -729,7 +729,7 @@ IModuleInterface* FModuleManager::LoadModuleWithFailureReason(const FName InModu
 }
 
 
-bool FModuleManager::UnloadModule( const FName InModuleName, bool bIsShutdown )
+bool FModuleManager::UnloadModule( const FName InModuleName, bool bIsShutdown, bool bAllowUnloadCode)
 {
 	// Do we even know about this module?
 	ModuleInfoPtr ModuleInfoPtr = FindModule(InModuleName);
@@ -756,7 +756,7 @@ bool FModuleManager::UnloadModule( const FName InModuleName, bool bIsShutdown )
 				// instead.  This makes it much less likely that code will be unloaded that could still be called by
 				// another module, such as a destructor or other virtual function.  The module will still be unloaded by
 				// the operating system when the process exits.
-				if( !bIsShutdown )
+				if( !bIsShutdown && bAllowUnloadCode )
 				{
 					// Unload the DLL
 					FPlatformProcess::FreeDllHandle( ModuleInfo.Handle );
diff --git a/Engine/Source/Runtime/Core/Public/Modules/ModuleManager.h b/Engine/Source/Runtime/Core/Public/Modules/ModuleManager.h
--- a/Engine/Source/Runtime/Core/Public/Modules/ModuleManager.h
+++ b/Engine/Source/Runtime/Core/Public/Modules/ModuleManager.h
@@ -288,10 +288,11 @@ public:
 	 *
 	 * @param InModuleName The name of the module to unload.  Should not include path, extension or platform/configuration info.  This is just the "module name" part of the module file name.
 	 * @param bIsShutdown Is this unload module call occurring at shutdown (default = false).
+	 * @param bAllowUnloadCode Allow unloading of code library if possible (default = true).
 	 * @return true if module was unloaded successfully, false otherwise.
 	 * @see AbandonModule, IsModuleLoaded, LoadModule, LoadModuleWithFailureReason
 	 */
-	CORE_API bool UnloadModule( const FName InModuleName, bool bIsShutdown = false );
+	CORE_API bool UnloadModule( const FName InModuleName, bool bIsShutdown = false, bool bAllowUnloadCode = true );
 
 	/**
 	 * Calls PreUnload then either unloads or abandons a module in memory, depending on whether the module supports unloading.
diff --git a/Engine/Source/Runtime/Projects/Private/ModuleDescriptor.cpp b/Engine/Source/Runtime/Projects/Private/ModuleDescriptor.cpp
--- a/Engine/Source/Runtime/Projects/Private/ModuleDescriptor.cpp
+++ b/Engine/Source/Runtime/Projects/Private/ModuleDescriptor.cpp
@@ -703,7 +703,7 @@ void FModuleDescriptor::LoadModulesForPhase(ELoadingPhase::Type LoadingPhase, co
 	}
 }
 
-void FModuleDescriptor::UnloadModulesForPhase(ELoadingPhase::Type LoadingPhase, const TArray<FModuleDescriptor>& Modules, TMap<FName, EModuleUnloadResult>& OutErrors, bool bSkipUnload /*= false*/)
+void FModuleDescriptor::UnloadModulesForPhase(ELoadingPhase::Type LoadingPhase, const TArray<FModuleDescriptor>& Modules, TMap<FName, EModuleUnloadResult>& OutErrors, bool bSkipUnload /*= false*/, bool bAllowUnloadCode /*= true*/)
 {
 	FScopedSlowTask SlowTask((float)Modules.Num());
 	for (const FModuleDescriptor& Descriptor : Modules)
@@ -734,7 +734,7 @@ void FModuleDescriptor::UnloadModulesForPhase(ELoadingPhase::Type LoadingPhase,
 		}
 
 		Module->PreUnloadCallback();
-		verify(FModuleManager::Get().UnloadModule(Descriptor.Name));
+		verify(FModuleManager::Get().UnloadModule(Descriptor.Name, false, bAllowUnloadCode));
 	}
 }
 
diff --git a/Engine/Source/Runtime/Projects/Private/PluginManager.cpp b/Engine/Source/Runtime/Projects/Private/PluginManager.cpp
--- a/Engine/Source/Runtime/Projects/Private/PluginManager.cpp
+++ b/Engine/Source/Runtime/Projects/Private/PluginManager.cpp
@@ -2755,10 +2755,10 @@ bool FPluginManager::TryLoadModulesForPlugin( const FPlugin& Plugin, const ELoad
 	return true;
 }
 
-bool FPluginManager::TryUnloadModulesForPlugin(const FPlugin& Plugin, const ELoadingPhase::Type LoadingPhase, FText* OutFailureMessage /*= nullptr*/, bool bSkipUnload /*= false*/) const
+bool FPluginManager::TryUnloadModulesForPlugin(const FPlugin& Plugin, const ELoadingPhase::Type LoadingPhase, FText* OutFailureMessage /*= nullptr*/, bool bSkipUnload /*= false*/, bool bAllowUnloadCode /*= true*/) const
 {
 	TMap<FName, EModuleUnloadResult> Errors;
-	FModuleDescriptor::UnloadModulesForPhase(LoadingPhase, Plugin.Descriptor.Modules, Errors, bSkipUnload);
+	FModuleDescriptor::UnloadModulesForPhase(LoadingPhase, Plugin.Descriptor.Modules, Errors, bSkipUnload, bAllowUnloadCode);
 
 	FText FailureMessage;
 	for( const TPair<FName, EModuleUnloadResult>& FailureIt : Errors)
@@ -3316,7 +3316,7 @@ bool FPluginManager::GetPluginDependencies_FromDescriptor(const FPluginReference
 	return false;
 }
 
-bool FPluginManager::UnmountExplicitlyLoadedPlugin(const FString& PluginName, FText* OutReason)
+bool FPluginManager::UnmountExplicitlyLoadedPlugin(const FString& PluginName, FText* OutReason, bool bAllowUnloadCode)
 {
 	TSharedPtr<FPlugin> Plugin = FindPluginInstance(PluginName);
 	TRACE_CPUPROFILER_EVENT_SCOPE(UnmountPluginFromExternalSource);
@@ -3395,7 +3395,7 @@ bool FPluginManager::UnmountExplicitlyLoadedPlugin(const FString& PluginName, FT
 	{
 		if (LoadingPhase != ELoadingPhase::None)
 		{
-			verify(TryUnloadModulesForPlugin(*Plugin, LoadingPhase));
+			verify(TryUnloadModulesForPlugin(*Plugin, LoadingPhase, nullptr, false, bAllowUnloadCode));
 		}
 	}
 
diff --git a/Engine/Source/Runtime/Projects/Private/PluginManager.h b/Engine/Source/Runtime/Projects/Private/PluginManager.h
--- a/Engine/Source/Runtime/Projects/Private/PluginManager.h
+++ b/Engine/Source/Runtime/Projects/Private/PluginManager.h
@@ -182,7 +182,8 @@ public:
 	virtual bool MountExplicitlyLoadedPlugin(const FString& PluginName) override;
 	virtual bool MountExplicitlyLoadedPlugin_FromFileName(const FString& PluginFileName) override;
 	virtual bool MountExplicitlyLoadedPlugin_FromDescriptor(const FPluginReferenceDescriptor& PluginDescriptor) override;
-	virtual bool UnmountExplicitlyLoadedPlugin(const FString& PluginName, FText* OutReason) override;
+	virtual bool UnmountExplicitlyLoadedPlugin(const FString& PluginName, FText* OutReason) override { return UnmountExplicitlyLoadedPlugin(PluginName, OutReason, true); }
+	virtual bool UnmountExplicitlyLoadedPlugin(const FString& PluginName, FText* OutReason, bool bAllowUnloadCode) override;
 	virtual bool GetPluginDependencies(const FString& PluginName, TArray<FPluginReferenceDescriptor>& PluginDependencies) override;
 	virtual bool GetPluginDependencies_FromFileName(const FString& PluginFileName, TArray<FPluginReferenceDescriptor>& PluginDependencies) override;
 	virtual bool GetPluginDependencies_FromDescriptor(const FPluginReferenceDescriptor& PluginDescriptor, TArray<FPluginReferenceDescriptor>& PluginDependencies) override;
@@ -262,7 +263,7 @@ private:
 	bool TryLoadModulesForPlugin(const FPlugin& Plugin, const ELoadingPhase::Type LoadingPhase) const;
 
 	/** Attempt to unload all the modules for the given plugin */
-	bool TryUnloadModulesForPlugin(const FPlugin& Plugin, const ELoadingPhase::Type LoadingPhase, FText* OutFailureMessage = nullptr, bool bSkipUnload = false) const;
+	bool TryUnloadModulesForPlugin(const FPlugin& Plugin, const ELoadingPhase::Type LoadingPhase, FText* OutFailureMessage = nullptr, bool bSkipUnload = false, bool bAllowUnloadCode = true) const;
 
 	/** Gets the instance of a given plugin */
 	TSharedPtr<FPlugin> FindPluginInstance(const FString& Name);
diff --git a/Engine/Source/Runtime/Projects/Public/Interfaces/IPluginManager.h b/Engine/Source/Runtime/Projects/Public/Interfaces/IPluginManager.h
--- a/Engine/Source/Runtime/Projects/Public/Interfaces/IPluginManager.h
+++ b/Engine/Source/Runtime/Projects/Public/Interfaces/IPluginManager.h
@@ -440,6 +440,7 @@ public:
 	 * Marks an explicitly loaded plugin as disabled, unmounts its content (does not work on plugins with compiled modules).
 	 */
 	virtual bool UnmountExplicitlyLoadedPlugin(const FString& PluginName, FText* OutReason) = 0;
+	virtual bool UnmountExplicitlyLoadedPlugin(const FString& PluginName, FText* OutReason, bool bAllowUnloadCode) = 0;
 
 	/**
 	 * Tries to get a list of plugin dependencies for a given plugin. Returns false if the plugin provided was not found
diff --git a/Engine/Source/Runtime/Projects/Public/ModuleDescriptor.h b/Engine/Source/Runtime/Projects/Public/ModuleDescriptor.h
--- a/Engine/Source/Runtime/Projects/Public/ModuleDescriptor.h
+++ b/Engine/Source/Runtime/Projects/Public/ModuleDescriptor.h
@@ -233,7 +233,7 @@ struct FModuleDescriptor
 	static PROJECTS_API void LoadModulesForPhase(ELoadingPhase::Type LoadingPhase, const TArray<FModuleDescriptor>& Modules, TMap<FName, EModuleLoadResult>& ModuleLoadErrors);
 
 	/** Unloads all the modules for a given loading phase. Returns a map of module names to load errors. bSkipUnload can be used to simulate unloading */
-	static PROJECTS_API void UnloadModulesForPhase(ELoadingPhase::Type LoadingPhase, const TArray<FModuleDescriptor>& Modules, TMap<FName, EModuleUnloadResult>& OutErrors, bool bSkipUnload = false);
+	static PROJECTS_API void UnloadModulesForPhase(ELoadingPhase::Type LoadingPhase, const TArray<FModuleDescriptor>& Modules, TMap<FName, EModuleUnloadResult>& OutErrors, bool bSkipUnload = false, bool bAllowUnloadCode = true);
 
 #if !IS_MONOLITHIC
 	/** Checks that all modules are compatible with the current engine version. Returns false and appends a list of names to OutIncompatibleFiles if not. */
