Support important messages for platform splash screen and improve plugin managers slow task reporting

diff --git a/Engine/Source/Editor/UnrealEd/Private/FeedbackContextEditor.cpp b/Engine/Source/Editor/UnrealEd/Private/FeedbackContextEditor.cpp
--- a/Engine/Source/Editor/UnrealEd/Private/FeedbackContextEditor.cpp
+++ b/Engine/Source/Editor/UnrealEd/Private/FeedbackContextEditor.cpp
@@ -598,14 +598,33 @@ void FFeedbackContextEditor::ProgressReported( const float TotalProgressInterp,
 	}
 	else if (FPlatformSplash::IsShown())
 	{
-		// Always show the top-most message
-		for (int i = ScopeStack.Num() - 1; i > -1; --i)
+		// look for important messages:
+		bool bFoundImportantMessage = false;
+		for (int32 i = ScopeStack.Num() - 1; i > -1; --i)
 		{
-			const FText ThisMessage = ScopeStack[i]->GetCurrentMessage();
-			if (!ThisMessage.IsEmpty())
+			if (ScopeStack[i]->Visibility == ESlowTaskVisibility::Important)
 			{
-				DisplayMessage = ThisMessage;
-				break;
+				const FText ThisMessage = ScopeStack[i]->GetCurrentMessage();
+				if (!ThisMessage.IsEmpty())
+				{
+					bFoundImportantMessage = true;
+					DisplayMessage = ThisMessage;
+					break;
+				}
+			}
+		}
+
+		// If nothing important, always show the top-most message
+		if (!bFoundImportantMessage)
+		{
+			for (int32 i = ScopeStack.Num() - 1; i > -1; --i)
+			{
+				const FText ThisMessage = ScopeStack[i]->GetCurrentMessage();
+				if (!ThisMessage.IsEmpty())
+				{
+					DisplayMessage = ThisMessage;
+					break;
+				}
 			}
 		}
 
diff --git a/Engine/Source/Runtime/Projects/Private/PluginManager.cpp b/Engine/Source/Runtime/Projects/Private/PluginManager.cpp
--- a/Engine/Source/Runtime/Projects/Private/PluginManager.cpp
+++ b/Engine/Source/Runtime/Projects/Private/PluginManager.cpp
@@ -2813,18 +2813,27 @@ bool FPluginManager::LoadModulesForEnabledPlugins( const ELoadingPhase::Type Loa
 	}
 	else
 	{
-		FScopedSlowTask SlowTask((float)AllPlugins.Num());
+		const FStringView LoadingPhaseAsString = FStringView(ELoadingPhase::ToString(LoadingPhase));
+		FScopedSlowTask SlowTask((float)AllPlugins.Num(),
+			FText::Format(LOCTEXT("LoadingModulesForEnabledPlugins", "Loading Plugin Modules for Phase {0}"), FText::FromStringView(LoadingPhaseAsString)));
+		SlowTask.Visibility = ESlowTaskVisibility::Important; // this function can be very slow, users will benefit from our messages
 		LLM_SCOPE_BYNAME(TEXT("Modules"));
 
 		// Load plugins!
+		int32 NumProcessedSinceProgress = 0;
 		for (const FDiscoveredPluginMap::ElementType& PluginPair : AllPlugins)
 		{
+			++NumProcessedSinceProgress;
 			const TSharedRef<FPlugin>& Plugin = DiscoveredPluginMapUtils::ResolvePluginFromMapVal(PluginPair.Value);
 
-			SlowTask.EnterProgressFrame(1);
-
 			if (Plugin->bEnabled && !Plugin->Descriptor.bExplicitlyLoaded)
 			{
+				const FText Message = FText::Format(LOCTEXT("LoadingModulesForPlugin", "Loading {0} Modules for Plugin: {1}"),
+					FText::FromStringView(LoadingPhaseAsString),
+					FText::FromString(Plugin->Name));
+				SlowTask.EnterProgressFrame((float)NumProcessedSinceProgress, Message);
+				NumProcessedSinceProgress = 0;
+
 				if (!TryLoadModulesForPlugin(Plugin.Get(), LoadingPhase))
 				{
 					bSuccess = false;
