Adding UE_LOG when the descriptor fails to read or parse as a bandaid. There is a race condition when calling FText::Format when FInternationalization is not initialized yet. The crash prevents the plugin name from outputing.

diff --git a/Engine/Source/Runtime/Projects/Private/PluginDescriptor.cpp b/Engine/Source/Runtime/Projects/Private/PluginDescriptor.cpp
--- a/Engine/Source/Runtime/Projects/Private/PluginDescriptor.cpp
+++ b/Engine/Source/Runtime/Projects/Private/PluginDescriptor.cpp
@@ -9,12 +9,16 @@
 
 #define LOCTEXT_NAMESPACE "PluginDescriptor"
 
+DEFINE_LOG_CATEGORY_STATIC(LogPluginDescriptor, Log, All);
+
 namespace PluginDescriptor
 {
 	bool ReadFile(const TCHAR* FileName, FString& Text, FText* OutFailReason = nullptr)
 	{
 		if (!FFileHelper::LoadFileToString(Text, FileName))
 		{
+			UE_LOG(LogPluginDescriptor, Error, TEXT("Failed to open descriptor file. %s"), FileName);
+
 			if (OutFailReason)
 			{
 				*OutFailReason = FText::Format(LOCTEXT("FailedToLoadDescriptorFile", "Failed to open descriptor file '{0}'"), FText::FromStringView(FileName));
@@ -43,6 +47,8 @@ namespace PluginDescriptor
 		TSharedRef<TJsonReader<>> Reader = TJsonReaderFactory<>::Create(Text);
 		if (!FJsonSerializer::Deserialize(Reader, JsonObject) || !JsonObject.IsValid())
 		{
+			UE_LOG(LogPluginDescriptor, Error, TEXT("Failed to read file. %s"), *Reader->GetErrorMessage());
+
 			if (OutFailReason)
 			{
 				*OutFailReason = FText::Format(LOCTEXT("FailedToReadDescriptorFile", "Failed to read file. {0}"), FText::FromString(Reader->GetErrorMessage()));
