In editor, optional plugin dependencies need to be compiled enabled or command line enabled otherwise they are ignored. Updated UBT to write out the BuildPlugins names to the target receipt or compiled into the executable.

diff --git a/Engine/Source/Developer/DesktopPlatform/Private/TargetReceipt.cpp b/Engine/Source/Developer/DesktopPlatform/Private/TargetReceipt.cpp
--- a/Engine/Source/Developer/DesktopPlatform/Private/TargetReceipt.cpp
+++ b/Engine/Source/Developer/DesktopPlatform/Private/TargetReceipt.cpp
@@ -169,6 +169,13 @@ bool FTargetReceipt::Read(const FString& FileName, bool bExpandVariables)
 		}
 	}
 
+	// Read the list of build plugins
+	FJsonSerializableArray BuildPluginsArray;
+	if (Object->TryGetStringArrayField(TEXT("BuildPlugins"), BuildPluginsArray))
+	{
+		BuildPlugins.Append(BuildPluginsArray.GetData(), BuildPluginsArray.Num());
+	}	
+
 	// Read the list of additional properties
 	const TArray<TSharedPtr<FJsonValue>>* AdditionalPropertiesArray;
 	if (Object->TryGetArrayField(TEXT("AdditionalProperties"), AdditionalPropertiesArray))
diff --git a/Engine/Source/Developer/DesktopPlatform/Public/TargetReceipt.h b/Engine/Source/Developer/DesktopPlatform/Public/TargetReceipt.h
--- a/Engine/Source/Developer/DesktopPlatform/Public/TargetReceipt.h
+++ b/Engine/Source/Developer/DesktopPlatform/Public/TargetReceipt.h
@@ -87,6 +87,9 @@ struct DESKTOPPLATFORM_API FTargetReceipt
 	/** All plugins that were either enabled or disabled via the target rules. */
 	TMap<FString, bool> PluginNameToEnabledState;
 
+	/** All plugins which are part of this target. */
+	TArray<FString> BuildPlugins;
+
 	/** Additional build properties passed through from the module rules */
 	TArray<FReceiptProperty> AdditionalProperties;
 
diff --git a/Engine/Source/Programs/UnrealBuildTool/Configuration/UEBuildTarget.cs b/Engine/Source/Programs/UnrealBuildTool/Configuration/UEBuildTarget.cs
--- a/Engine/Source/Programs/UnrealBuildTool/Configuration/UEBuildTarget.cs
+++ b/Engine/Source/Programs/UnrealBuildTool/Configuration/UEBuildTarget.cs
@@ -2377,6 +2377,12 @@ namespace UnrealBuildTool
 				Receipt.PluginNameToEnabledState[DisabledPluginName] = false;
 			}
 
+			// Add all the build plugin names
+			foreach (UEBuildPlugin Plugin in BuildPlugins!)
+			{
+				Receipt.BuildPlugins.Add(Plugin.Name);
+			}
+
 			// Find all the modules which are part of this target
 			HashSet<UEBuildModule> UniqueLinkedModules = new HashSet<UEBuildModule>();
 			foreach (UEBuildBinary Binary in Binaries)
@@ -4422,6 +4428,28 @@ namespace UnrealBuildTool
 			SetupUbtPlugins(Logger);
 		}
 
+		/// <summary>
+		/// Modify a module rules based on the current target, be very careful not to violate the shared environment!
+		/// </summary>
+		/// <param name="moduleName">name of the module</param>
+		/// <param name="moduleRules">the module's rules</param>
+		private void ModifyModuleRulesForTarget(string moduleName, ModuleRules moduleRules)
+		{
+			if (moduleName == "Projects" && (Rules.Type == TargetType.Editor))
+			{
+				// Monolithic and non-shared environment builds compile in the build plugins. Non-monolithic editor builds save them in the BuildPlugins receipt to avoid invalidating the shared build environment.
+				// See Projects.Build.cs
+				bool bUsingTargetReceipt = (Rules.BuildEnvironment == TargetBuildEnvironment.Shared) && (Rules.LinkType != TargetLinkType.Monolithic);
+				if (!bUsingTargetReceipt)
+				{
+					// Compile in the build plugins list for the runtime plugin manager to filter optional plugin references against. Only plugins that were actually built
+					// should be enabled when another plugin has an optional dependency on the plugin.
+					IEnumerable<string> BuildPluginStrings = BuildPlugins?.Select(x => $"TEXT(\"{x}\")") ?? Array.Empty<string>();
+					moduleRules.PrivateDefinitions.Add($"UBT_TARGET_BUILD_PLUGINS={String.Join(", ", BuildPluginStrings)}");
+				}
+			}
+		}
+
 		/// <summary>
 		/// Creates a plugin instance from a reference to it
 		/// </summary>
@@ -5541,6 +5569,9 @@ namespace UnrealBuildTool
 					}
 				}
 
+				// Allow the current target to modify the module rules
+				ModifyModuleRulesForTarget(ModuleName, RulesObject);
+
 				// Allow the current platform to modify the module rules
 				UEBuildPlatform.GetBuildPlatform(Platform).ModifyModuleRulesForActivePlatform(ModuleName, RulesObject, Rules);
 
diff --git a/Engine/Source/Programs/UnrealBuildTool/System/TargetReceipt.cs b/Engine/Source/Programs/UnrealBuildTool/System/TargetReceipt.cs
--- a/Engine/Source/Programs/UnrealBuildTool/System/TargetReceipt.cs
+++ b/Engine/Source/Programs/UnrealBuildTool/System/TargetReceipt.cs
@@ -327,6 +327,11 @@ namespace UnrealBuildTool
 		/// </summary>
 		public Dictionary<string, bool> PluginNameToEnabledState = new Dictionary<string, bool>();
 
+		/// <summary>
+		/// All plugins that were built via the target rules.
+		/// </summary>
+		public List<string> BuildPlugins = new List<string>();
+
 		/// <summary>
 		/// Additional build properties passed through from the module rules
 		/// </summary>
@@ -393,6 +398,13 @@ namespace UnrealBuildTool
 					PluginNameToEnabledState.Add(Pair.Key, Pair.Value);
 				}
 			}
+			foreach (string PluginName in Other.BuildPlugins)
+			{
+				if (!BuildPlugins.Contains(PluginName))
+				{
+					BuildPlugins.Add(PluginName);
+				}
+			}
 		}
 
 		/// <summary>
@@ -651,6 +663,13 @@ namespace UnrealBuildTool
 				}
 			}
 
+			// Read the build plugins
+			string[]? BuildPlugins;
+			if (RawObject.TryGetStringArrayField("BuildPlugins", out BuildPlugins))
+			{
+				Receipt.BuildPlugins.AddAll(BuildPlugins);
+			}
+
 			// Read the additional properties
 			JsonObject[]? AdditionalPropertyObjects;
 			if (RawObject.TryGetObjectArrayField("AdditionalProperties", out AdditionalPropertyObjects))
@@ -788,6 +807,11 @@ namespace UnrealBuildTool
 					Writer.WriteArrayEnd();
 				}
 
+				if (BuildPlugins.Count > 0)
+				{
+					Writer.WriteStringArrayField("BuildPlugins", BuildPlugins.OrderBy(x => x));
+				}
+
 				if (AdditionalProperties.Count > 0)
 				{
 					Writer.WriteArrayStart("AdditionalProperties");
diff --git a/Engine/Source/Runtime/Projects/Private/PluginManager.cpp b/Engine/Source/Runtime/Projects/Private/PluginManager.cpp
--- a/Engine/Source/Runtime/Projects/Private/PluginManager.cpp
+++ b/Engine/Source/Runtime/Projects/Private/PluginManager.cpp
@@ -1406,6 +1406,75 @@ bool FPluginManager::ConfigureEnabledPlugins()
 		// Keep the list of newly available localization targets
 		TArray<FString> AdditionalLocResPaths;
 
+#if READ_TARGET_ENABLED_PLUGINS_FROM_RECEIPT
+		FString DefaultEditorTarget;
+		GConfig->GetString(TEXT("/Script/BuildSettings.BuildSettings"), TEXT("DefaultEditorTarget"), DefaultEditorTarget, GEngineIni);
+
+		auto FindFirstMatchingTargetFile = [&DefaultEditorTarget](const FString& ReceiptWildcard) -> TUniquePtr<FTargetReceipt>
+		{
+			TArray<FString> AllTargetFilesWithoutPath;
+			const FString ReceiptPath = FPaths::GetPath(ReceiptWildcard);
+			IFileManager::Get().FindFiles(AllTargetFilesWithoutPath, *ReceiptWildcard, true, false);
+
+			for (const FString& TargetFileWithoutPath : AllTargetFilesWithoutPath)
+			{
+				const FString TargetFile = FPaths::Combine(ReceiptPath, TargetFileWithoutPath);
+				TUniquePtr<FTargetReceipt> Receipt = MakeUnique<FTargetReceipt>();
+				if (Receipt->Read(TargetFile))
+				{
+					if (Receipt->TargetType == FApp::GetBuildTargetType() && Receipt->Configuration == FApp::GetBuildConfiguration())
+					{
+						bool bIsDefaultTarget = Receipt->TargetType != EBuildTargetType::Editor || (DefaultEditorTarget.Len() == 0) || (DefaultEditorTarget == Receipt->TargetName);
+						if (bIsDefaultTarget)
+						{
+							return Receipt;
+						}
+					}
+				}
+			}
+			return TUniquePtr<FTargetReceipt>();
+		};
+#endif // READ_TARGET_ENABLED_PLUGINS_FROM_RECEIPT
+
+#if !WITH_EDITOR
+		// const to ensure it to stays empty
+		const TSet<FString> AllowedOptionalDependencies;
+#else
+		// Set of all the plugin names that are allowed to be enabled for a plugin with optional dependencies. Only read in Editor
+		TSet<FString> AllowedOptionalDependencies;
+
+#if READ_TARGET_ENABLED_PLUGINS_FROM_RECEIPT
+		{
+			SCOPED_BOOT_TIMING("ReadTargetBuildPluginsFromReceipt");
+
+			// Read the build plugins from the target file using the target receipt file. This controls which optional plugin references can be enabled.
+			auto ReadBuildPluginsFromFirstMatchingTargetFile = [&FindFirstMatchingTargetFile, &AllowedOptionalDependencies](const TCHAR* BaseDir) -> bool
+			{
+				const FString ReceiptWildcard = FTargetReceipt::GetDefaultPath(BaseDir, TEXT("*"), FPlatformProcess::GetBinariesSubdirectory(), FApp::GetBuildConfiguration(), nullptr);
+
+				TUniquePtr<FTargetReceipt> Receipt = FindFirstMatchingTargetFile(ReceiptWildcard);
+				if (Receipt.IsValid())
+				{
+					AllowedOptionalDependencies.Append(Receipt->BuildPlugins);
+					return true;
+				}
+				return false;
+			};
+
+			if (!ReadBuildPluginsFromFirstMatchingTargetFile(FPlatformMisc::ProjectDir()))
+			{
+				ReadBuildPluginsFromFirstMatchingTargetFile(FPlatformMisc::EngineDir());
+			}
+		}
+#else
+		{
+			// Configure the plugins that were enabled from the target file using defines
+			SCOPED_BOOT_TIMING("ReadTargetBuildPlugins");
+			AllowedOptionalDependencies.Append({ UBT_TARGET_BUILD_PLUGINS });
+		}
+#endif // READ_TARGET_ENABLED_PLUGINS_FROM_RECEIPT
+#endif // !WITH_EDITOR
+
 		// Check which plugins have been enabled or excluded via the command line
 		{
 			SCOPED_BOOT_TIMING("ParseCmdLineForPlugins");
@@ -1468,6 +1537,10 @@ bool FPluginManager::ConfigureEnabledPlugins()
 			}
 			if (ExtraPluginsToEnable.Num() > 0)
 			{
+#if WITH_EDITOR
+				AllowedOptionalDependencies.Append(ExtraPluginsToEnable);
+#endif // WITH_EDITOR
+
 				auto IsRestrictedPlugin = [this](const FString& PluginName)
 				{
 					if (TSharedPtr<IPlugin> PluginPtr = FindPlugin(PluginName))
@@ -1485,7 +1558,7 @@ bool FPluginManager::ConfigureEnabledPlugins()
 					if (!ConfiguredPluginNames.Contains(EnablePluginName) && !ExceptPlugins.Contains(EnablePluginName) && (!bExceptRestrictedPlugins || !IsRestrictedPlugin(EnablePluginName)))
 					{
 						if (!ConfigureEnabledPluginForCurrentTarget(FPluginReferenceDescriptor(EnablePluginName, true), EnabledPlugins,
-							bAllPluginsEnabledViaCommandLine ? TEXTVIEW("Commandline EnableAllPlugins") : TEXTVIEW("CommandLine EnablePlugins=")))
+							bAllPluginsEnabledViaCommandLine ? TEXTVIEW("Commandline EnableAllPlugins") : TEXTVIEW("CommandLine EnablePlugins="), AllowedOptionalDependencies))
 						{
 							if (bAllPluginsEnabledViaCommandLine)
 							{
@@ -1507,7 +1580,7 @@ bool FPluginManager::ConfigureEnabledPlugins()
 				if (!ConfiguredPluginNames.Contains(DisablePluginName))
 				{
 					if (!ConfigureEnabledPluginForCurrentTarget(FPluginReferenceDescriptor(DisablePluginName, false), EnabledPlugins,
-						TEXTVIEW("CommandLine DisablePlugins=")))
+						TEXTVIEW("CommandLine DisablePlugins="), AllowedOptionalDependencies))
 					{
 						return false;
 					}
@@ -1522,49 +1595,31 @@ bool FPluginManager::ConfigureEnabledPlugins()
 
 #if READ_TARGET_ENABLED_PLUGINS_FROM_RECEIPT
 			// Configure the plugins that were enabled or disabled from the target file using the target receipt file
-			FString DefaultEditorTarget;
-			GConfig->GetString(TEXT("/Script/BuildSettings.BuildSettings"), TEXT("DefaultEditorTarget"), DefaultEditorTarget, GEngineIni);
-
-			auto ConfigurePluginsFromFirstMatchingTargetFile = [this, &ConfiguredPluginNames, &EnabledPlugins, &DefaultEditorTarget](const TCHAR* BaseDir, bool& bOutError) -> bool
+			auto ConfigurePluginsFromFirstMatchingTargetFile = [this, &FindFirstMatchingTargetFile, &ConfiguredPluginNames, &EnabledPlugins, &AllowedOptionalDependencies](const TCHAR* BaseDir, bool& bOutError) -> bool
 			{
-				TArray<FString> AllTargetFilesWithoutPath;
 				const FString ReceiptWildcard = FTargetReceipt::GetDefaultPath(BaseDir, TEXT("*"), FPlatformProcess::GetBinariesSubdirectory(), FApp::GetBuildConfiguration(), nullptr);
-				const FString ReceiptPath = FPaths::GetPath(ReceiptWildcard);
-				IFileManager::Get().FindFiles(AllTargetFilesWithoutPath, *ReceiptWildcard, true, false);
 				FString SourceDescription = FString::Printf(TEXT("Receipt files %s"), *ReceiptWildcard);
 
-				for (const FString& TargetFileWithoutPath : AllTargetFilesWithoutPath)
+				TUniquePtr<FTargetReceipt> Receipt = FindFirstMatchingTargetFile(ReceiptWildcard);
+				if (Receipt.IsValid())
 				{
-					const FString TargetFile = FPaths::Combine(ReceiptPath, TargetFileWithoutPath);
-					FTargetReceipt Receipt;
-					if (Receipt.Read(TargetFile))
+					for (const TPair<FString, bool>& Pair : Receipt->PluginNameToEnabledState)
 					{
-						if (Receipt.TargetType == FApp::GetBuildTargetType() && Receipt.Configuration == FApp::GetBuildConfiguration())
+						const FString& PluginName = Pair.Key;
+						const bool bEnabled = Pair.Value;
+						if (!ConfiguredPluginNames.Contains(PluginName))
 						{
-							bool bIsDefaultTarget = Receipt.TargetType != EBuildTargetType::Editor || (DefaultEditorTarget.Len() == 0) || (DefaultEditorTarget == Receipt.TargetName);
-
-							if (bIsDefaultTarget)
+							if (!ConfigureEnabledPluginForCurrentTarget(FPluginReferenceDescriptor(PluginName, bEnabled), EnabledPlugins,
+								SourceDescription, AllowedOptionalDependencies))
 							{
-								for (const TPair<FString, bool>& Pair : Receipt.PluginNameToEnabledState)
-								{
-									const FString& PluginName = Pair.Key;
-									const bool bEnabled = Pair.Value;
-									if (!ConfiguredPluginNames.Contains(PluginName))
-									{
-										if (!ConfigureEnabledPluginForCurrentTarget(FPluginReferenceDescriptor(PluginName, bEnabled), EnabledPlugins,
-											SourceDescription))
-										{
-											bOutError = true;
-											break;
-										}
-										ConfiguredPluginNames.Add(PluginName);
-									}
-								}
-
-								return true;
+								bOutError = true;
+								break;
 							}
+							ConfiguredPluginNames.Add(PluginName);
 						}
 					}
+
+					return true;
 				}
 
 				return false;
@@ -1593,7 +1648,7 @@ bool FPluginManager::ConfigureEnabledPlugins()
 					if (!ConfiguredPluginNames.Contains(TargetEnabledPlugin))
 					{
 						if (!ConfigureEnabledPluginForCurrentTarget(FPluginReferenceDescriptor(TargetEnabledPlugin, true), EnabledPlugins,
-							TEXTVIEW("UBT_TARGET_ENABLED_PLUGINS")))
+							TEXTVIEW("UBT_TARGET_ENABLED_PLUGINS"), AllowedOptionalDependencies))
 						{
 							return false;
 						}
@@ -1611,7 +1666,7 @@ bool FPluginManager::ConfigureEnabledPlugins()
 					if (!ConfiguredPluginNames.Contains(TargetDisabledPlugin))
 					{
 						if (!ConfigureEnabledPluginForCurrentTarget(FPluginReferenceDescriptor(TargetDisabledPlugin, false), EnabledPlugins,
-							TEXTVIEW("UBT_TARGET_ENABLED_PLUGINS")))
+							TEXTVIEW("UBT_TARGET_ENABLED_PLUGINS"), AllowedOptionalDependencies))
 						{
 							return false;
 						}
@@ -1622,14 +1677,14 @@ bool FPluginManager::ConfigureEnabledPlugins()
 #endif // READ_TARGET_ENABLED_PLUGINS_FROM_RECEIPT
 
 			auto ProcessPluginConfigurations =
-				[&ConfiguredPluginNames, &EnabledPlugins, this]
+				[&ConfiguredPluginNames, &EnabledPlugins, &AllowedOptionalDependencies, this]
 				(const TArray<FPluginReferenceDescriptor>& PluginReferences, FStringView SourceDescription)->bool
 			{
 				for (const FPluginReferenceDescriptor& PluginReference : PluginReferences)
 				{
 					if (!ConfiguredPluginNames.Contains(PluginReference.Name))
 					{
-						if (!ConfigureEnabledPluginForCurrentTarget(PluginReference, EnabledPlugins, SourceDescription))
+						if (!ConfigureEnabledPluginForCurrentTarget(PluginReference, EnabledPlugins, SourceDescription, AllowedOptionalDependencies))
 						{
 							return false;
 						}
@@ -1670,7 +1725,7 @@ bool FPluginManager::ConfigureEnabledPlugins()
 					if (Plugin->IsEnabledByDefault(bAllowEnginePluginsEnabledByDefault) && !ConfiguredPluginNames.Contains(PluginName))
 					{
 						if (!ConfigureEnabledPluginForCurrentTarget(FPluginReferenceDescriptor(PluginName, true),
-							EnabledPlugins, SourceDescription))
+							EnabledPlugins, SourceDescription, AllowedOptionalDependencies))
 						{
 							return false;
 						}
@@ -1681,7 +1736,7 @@ bool FPluginManager::ConfigureEnabledPlugins()
 		}
 
 #if IS_PROGRAM
-		auto EnableProgramPlugin = [this, &ConfiguredPluginNames, &EnabledPlugins](const TCHAR* ConfigEntry, bool bOptionalPlugin) mutable
+		auto EnableProgramPlugin = [this, &ConfiguredPluginNames, &EnabledPlugins, &AllowedOptionalDependencies](const TCHAR* ConfigEntry, bool bOptionalPlugin) mutable
 		{
 			TArray<FString> ProgramPluginNames;
 			GConfig->GetArray(TEXT("Plugins"), ConfigEntry, ProgramPluginNames, GEngineIni);
@@ -1694,7 +1749,7 @@ bool FPluginManager::ConfigureEnabledPlugins()
 					FPluginReferenceDescriptor PluginReference(PluginName, true);
 					PluginReference.bOptional = bOptionalPlugin;
 					if (!ConfigureEnabledPluginForCurrentTarget(PluginReference,
-						EnabledPlugins, SourceDescription))
+						EnabledPlugins, SourceDescription, AllowedOptionalDependencies))
 					{
 						return false;
 					}
@@ -2010,14 +2065,14 @@ bool FPluginManager::RequiresTempTargetForCodePlugin(const FProjectDescriptor* P
 	FConfigurePluginResultInfo ResultInfo;
 
 	TSet<FString> ProjectCodePlugins;
-	if (!GetCodePluginsForProject(ProjectDescriptor, Platform, Configuration, TargetType, AllPlugins, ProjectCodePlugins, ResultInfo))
+	if (!GetCodePluginsForProject(ProjectDescriptor, Platform, Configuration, TargetType, AllPlugins, ProjectCodePlugins, TSet<FString>(), ResultInfo))
 	{
 		OutReason = FText::Format(LOCTEXT("TempTarget_MissingPluginForTarget", "{0} plugin is referenced by target but not found"), FText::FromString(ResultInfo.PluginReference->Name));
 		return true;
 	}
 
 	TSet<FString> DefaultCodePlugins;
-	if (!GetCodePluginsForProject(nullptr, Platform, Configuration, TargetType, AllPlugins, DefaultCodePlugins, ResultInfo))
+	if (!GetCodePluginsForProject(nullptr, Platform, Configuration, TargetType, AllPlugins, DefaultCodePlugins, TSet<FString>(), ResultInfo))
 	{
 		OutReason = FText::Format(LOCTEXT("TempTarget_MissingPluginForDefaultTarget", "{0} plugin is referenced by the default target but not found"), FText::FromString(ResultInfo.PluginReference->Name));
 		return true;
@@ -2046,7 +2101,7 @@ bool FPluginManager::RequiresTempTargetForCodePlugin(const FProjectDescriptor* P
 
 bool FPluginManager::GetCodePluginsForProject(const FProjectDescriptor* ProjectDescriptor, const FString& Platform,
 	EBuildConfiguration Configuration, EBuildTargetType TargetType, FDiscoveredPluginMap& AllPlugins,
-	TSet<FString>& CodePluginNames, FConfigurePluginResultInfo& OutResultInfo)
+	TSet<FString>& CodePluginNames, const TSet<FString>& AllowedOptionalDependencies, FConfigurePluginResultInfo& OutResultInfo)
 {
 	// Can only check the current project at the moment, since we won't have enumerated them otherwise
 	check(ProjectDescriptor == nullptr || ProjectDescriptor == IProjectManager::Get().GetCurrentProject());
@@ -2074,7 +2129,7 @@ bool FPluginManager::GetCodePluginsForProject(const FProjectDescriptor* ProjectD
 			{
 				if (!ConfigureEnabledPluginForTarget(PluginReference, ProjectDescriptor, FString(), Platform,
 					Configuration, TargetType, bLoadPluginsForTargetPlatforms, AllPlugins, EnabledPlugins,
-					OutResultInfo))
+					AllowedOptionalDependencies, OutResultInfo))
 				{
 					return false;
 				}
@@ -2090,7 +2145,7 @@ bool FPluginManager::GetCodePluginsForProject(const FProjectDescriptor* ProjectD
 		{
 			if (!ConfigureEnabledPluginForTarget(FPluginReferenceDescriptor(PluginPair.Key, true), ProjectDescriptor,
 				FString(), Platform, Configuration, TargetType, bLoadPluginsForTargetPlatforms, AllPlugins,
-				EnabledPlugins, OutResultInfo))
+				EnabledPlugins, AllowedOptionalDependencies, OutResultInfo))
 			{
 				return false;
 			}
@@ -2117,7 +2172,7 @@ bool FPluginManager::GetCodePluginsForProject(const FProjectDescriptor* ProjectD
 }
 
 bool FPluginManager::ConfigureEnabledPluginForCurrentTarget(const FPluginReferenceDescriptor& FirstReference,
-	TMap<FString, FPlugin*>& EnabledPlugins, FStringView SourceOfPluginRequest)
+	TMap<FString, FPlugin*>& EnabledPlugins, FStringView SourceOfPluginRequest, const TSet<FString>& AllowedOptionalDependencies)
 {
 	SCOPED_BOOT_TIMING("ConfigureEnabledPluginForCurrentTarget");
 
@@ -2125,7 +2180,7 @@ bool FPluginManager::ConfigureEnabledPluginForCurrentTarget(const FPluginReferen
 
 	if (ConfigureEnabledPluginForTarget(FirstReference, IProjectManager::Get().GetCurrentProject(), UE_APP_NAME,
 		FPlatformMisc::GetUBTPlatform(), FApp::GetBuildConfiguration(), FApp::GetBuildTargetType(),
-		(bool)LOAD_PLUGINS_FOR_TARGET_PLATFORMS, AllPlugins, EnabledPlugins, ResultInfo))
+		(bool)LOAD_PLUGINS_FOR_TARGET_PLATFORMS, AllPlugins, EnabledPlugins, AllowedOptionalDependencies, ResultInfo))
 	{
 		return true;
 	}
@@ -2286,6 +2341,7 @@ bool FPluginManager::ConfigureEnabledPluginForTarget(const FPluginReferenceDescr
 	const FProjectDescriptor* ProjectDescriptor, const FString& TargetName, const FString& Platform,
 	EBuildConfiguration Configuration, EBuildTargetType TargetType, bool bLoadPluginsForTargetPlatforms,
 	FDiscoveredPluginMap& AllPlugins, TMap<FString, FPlugin*>& EnabledPlugins,
+	const TSet<FString>& AllowedOptionalDependencies,
 	FConfigurePluginResultInfo& OutResultInfo)
 {
 	if (EnabledPlugins.Contains(FirstReference.Name))
@@ -2486,6 +2542,15 @@ bool FPluginManager::ConfigureEnabledPluginForTarget(const FPluginReferenceDescr
 				return false;
 			}
 
+#if WITH_EDITOR
+			// Allowed optional plugins are compiled enabled or enabled via the commandline. Ignore those that are not.
+			if (NextReference.bOptional && !AllowedOptionalDependencies.IsEmpty() && !AllowedOptionalDependencies.Contains(NextReference.Name))
+			{
+				UE_LOG(LogPluginManager, Display, TEXT("Ignored optional reference to '%s' plugin from '%s' plugin; plugin was not built by target."), *NextReference.Name, *Plugin.GetName());
+				continue;
+			}
+#endif
+
 			if (!EnabledPlugins.Contains(NextReference.Name) && bIsNewlySeen)
 			{
 				NewPluginQueue.Add(&NextReference);
diff --git a/Engine/Source/Runtime/Projects/Private/PluginManager.h b/Engine/Source/Runtime/Projects/Private/PluginManager.h
--- a/Engine/Source/Runtime/Projects/Private/PluginManager.h
+++ b/Engine/Source/Runtime/Projects/Private/PluginManager.h
@@ -223,17 +223,17 @@ private:
 	static void FindPluginManifestsInDirectory(const FString& PluginManifestDirectory, TArray<FString>& FileNames);
 
 	/** Gets all the code plugins that are enabled for a content only project */
-	static bool GetCodePluginsForProject(const FProjectDescriptor* ProjectDescriptor, const FString& Platform, EBuildConfiguration Configuration, EBuildTargetType TargetType, FDiscoveredPluginMap& AllPlugins, TSet<FString>& CodePluginNames, FConfigurePluginResultInfo& OutResultInfo);
+	static bool GetCodePluginsForProject(const FProjectDescriptor* ProjectDescriptor, const FString& Platform, EBuildConfiguration Configuration, EBuildTargetType TargetType, FDiscoveredPluginMap& AllPlugins, TSet<FString>& CodePluginNames, const TSet<FString>& AllowedOptionalPluginReferences, FConfigurePluginResultInfo& OutResultInfo);
 
 	/** Sets the bPluginEnabled flag on all plugins found from DiscoverAllPlugins that are enabled in config */
 	bool ConfigureEnabledPlugins();
 
-	/** Adds a single enabled plugin, and all its dependencies */
+	/** Adds a single enabled plugin, and all its dependencies. If AllowedOptionalDependencies is not empty, only enable optional dependencies that are listed */
 	bool ConfigureEnabledPluginForCurrentTarget(const FPluginReferenceDescriptor& FirstReference,
-		TMap<FString, FPlugin*>& EnabledPlugins, FStringView SourceOfPluginRequest);
+		TMap<FString, FPlugin*>& EnabledPlugins, FStringView SourceOfPluginRequest, const TSet<FString>& AllowedOptionalDependencies);
 
-	/** Adds a single enabled plugin and all its dependencies. */
-	static bool ConfigureEnabledPluginForTarget(const FPluginReferenceDescriptor& FirstReference, const FProjectDescriptor* ProjectDescriptor, const FString& TargetName, const FString& Platform, EBuildConfiguration Configuration, EBuildTargetType TargetType, bool bLoadPluginsForTargetPlatforms, FDiscoveredPluginMap& AllPlugins, TMap<FString, FPlugin*>& EnabledPlugins, FConfigurePluginResultInfo& OutResultInfo);
+	/** Adds a single enabled plugin and all its dependencies. If AllowedOptionalDependencies is not empty, only enable optional dependencies that are listed */
+	static bool ConfigureEnabledPluginForTarget(const FPluginReferenceDescriptor& FirstReference, const FProjectDescriptor* ProjectDescriptor, const FString& TargetName, const FString& Platform, EBuildConfiguration Configuration, EBuildTargetType TargetType, bool bLoadPluginsForTargetPlatforms, FDiscoveredPluginMap& AllPlugins, TMap<FString, FPlugin*>& EnabledPlugins, const TSet<FString>& AllowedOptionalDependencies, FConfigurePluginResultInfo& OutResultInfo);
 
 	/** Prompts the user to download a missing plugin from the given URL */
 	static bool PromptToDownloadPlugin(const FString& PluginName, const FString& MarketplaceURL);
