Keep shared reference of IPlugin instead of its raw pointer

diff --git a/Engine/Source/Runtime/Projects/Public/Interfaces/IPluginManager.h b/Engine/Source/Runtime/Projects/Public/Interfaces/IPluginManager.h
--- a/Engine/Source/Runtime/Projects/Public/Interfaces/IPluginManager.h
+++ b/Engine/Source/Runtime/Projects/Public/Interfaces/IPluginManager.h
@@ -5,6 +5,7 @@
 #include "CoreMinimal.h"
 #include "PluginDescriptor.h"
 #include "Containers/VersePathFwd.h"
+#include "Templates/SharedPointer.h"
 
 struct FProjectDescriptor;
 class FJsonObject;
@@ -69,7 +70,7 @@ struct FPluginStatus
 /**
  * Information about an enabled plugin.
  */
-class IPlugin
+class IPlugin : public TSharedFromThis<IPlugin>
 {
 public:
 	/* Virtual destructor */
