Added new function suite to return found plugins that are enabled. Updated the LoadBuiltInGameFeaturePlugins sort function to skip plugin dependencies that are not enabled for correctness however it does not impact the final EnabledPlugins list

diff --git a/Engine/Plugins/Runtime/GameFeatures/Source/GameFeatures/Private/GameFeaturesSubsystem.cpp b/Engine/Plugins/Runtime/GameFeatures/Source/GameFeatures/Private/GameFeaturesSubsystem.cpp
--- a/Engine/Plugins/Runtime/GameFeatures/Source/GameFeatures/Private/GameFeaturesSubsystem.cpp
+++ b/Engine/Plugins/Runtime/GameFeatures/Source/GameFeatures/Private/GameFeaturesSubsystem.cpp
@@ -1901,7 +1901,7 @@ void UGameFeaturesSubsystem::LoadBuiltInGameFeaturePlugins(FBuiltInPluginAdditio
 		{
 			if (Dependency.bEnabled)
 			{
-				if (TSharedPtr<IPlugin> FoundPlugin = PluginManager.FindPlugin(Dependency.Name))
+				if (TSharedPtr<IPlugin> FoundPlugin = PluginManager.FindEnabledPlugin(Dependency.Name))
 				{
 					Dependencies.Add(FoundPlugin.ToSharedRef());
 				}
diff --git a/Engine/Source/Runtime/Projects/Private/PluginManager.cpp b/Engine/Source/Runtime/Projects/Private/PluginManager.cpp
--- a/Engine/Source/Runtime/Projects/Private/PluginManager.cpp
+++ b/Engine/Source/Runtime/Projects/Private/PluginManager.cpp
@@ -3004,6 +3004,45 @@ TSharedPtr<IPlugin> FPluginManager::FindPluginFromDescriptor(const FPluginRefere
 	}
 }
 
+TSharedPtr<IPlugin> FPluginManager::FindEnabledPlugin(const FStringView Name)
+{
+	const TSharedPtr<IPlugin> Plugin = FindPlugin(Name);
+	if (Plugin.IsValid() && Plugin->IsEnabled())
+	{
+		return Plugin;
+	}
+	else
+	{
+		return TSharedPtr<IPlugin>();
+	}
+}
+
+TSharedPtr<IPlugin> FPluginManager::FindEnabledPluginFromPath(const FString& PluginPath)
+{
+	const TSharedPtr<IPlugin> Plugin = FindPluginFromPath(PluginPath);
+	if (Plugin.IsValid() && Plugin->IsEnabled())
+	{
+		return Plugin;
+	}
+	else
+	{
+		return TSharedPtr<IPlugin>();
+	}
+}
+
+TSharedPtr<IPlugin> FPluginManager::FindEnabledPluginFromDescriptor(const FPluginReferenceDescriptor& PluginDesc)
+{
+	const TSharedPtr<IPlugin> Plugin = FindPluginFromDescriptor(PluginDesc);
+	if (Plugin.IsValid() && Plugin->IsEnabled())
+	{
+		return Plugin;
+	}
+	else
+	{
+		return TSharedPtr<IPlugin>();
+	}
+}
+
 TArray<TSharedRef<IPlugin>> FPluginManager::GetEnabledPlugins()
 {
 	TArray<TSharedRef<IPlugin>> Plugins;
diff --git a/Engine/Source/Runtime/Projects/Private/PluginManager.h b/Engine/Source/Runtime/Projects/Private/PluginManager.h
--- a/Engine/Source/Runtime/Projects/Private/PluginManager.h
+++ b/Engine/Source/Runtime/Projects/Private/PluginManager.h
@@ -157,6 +157,15 @@ public:
 	virtual TSharedPtr<IPlugin> FindPluginFromPath(const FString& PluginPath) override;
 	virtual TSharedPtr<IPlugin> FindPluginFromDescriptor(const FPluginReferenceDescriptor& PluginDesc) override;
 
+	virtual TSharedPtr<IPlugin> FindEnabledPlugin(const FStringView Name) override;
+	virtual TSharedPtr<IPlugin> FindEnabledPlugin(const ANSICHAR* Name) override
+	{
+		FString NameString(Name);
+		return FindEnabledPlugin(FStringView(NameString));
+	}
+	virtual TSharedPtr<IPlugin> FindEnabledPluginFromPath(const FString& PluginPath) override;
+	virtual TSharedPtr<IPlugin> FindEnabledPluginFromDescriptor(const FPluginReferenceDescriptor& PluginDesc) override;
+
 	virtual void FindPluginsUnderDirectory(const FString& Directory, TArray<FString>& OutPluginFilePaths) override;
 
 	virtual TArray<TSharedRef<IPlugin>> GetEnabledPlugins() override;
diff --git a/Engine/Source/Runtime/Projects/Public/Interfaces/IPluginManager.h b/Engine/Source/Runtime/Projects/Public/Interfaces/IPluginManager.h
--- a/Engine/Source/Runtime/Projects/Public/Interfaces/IPluginManager.h
+++ b/Engine/Source/Runtime/Projects/Public/Interfaces/IPluginManager.h
@@ -321,6 +321,17 @@ public:
 	virtual TSharedPtr<IPlugin> FindPluginFromPath(const FString& PluginPath) = 0;
 	virtual TSharedPtr<IPlugin> FindPluginFromDescriptor(const FPluginReferenceDescriptor& PluginDesc) = 0;
 
+	/**
+	 * Finds information for an enabled plugin.
+	 *
+	 * @return	 Pointer to the enabled plugin's information, or nullptr if not enabled or can't be found.
+	 */
+	virtual TSharedPtr<IPlugin> FindEnabledPlugin(const FStringView Name) = 0;
+	virtual TSharedPtr<IPlugin> FindEnabledPlugin(const ANSICHAR* Name) = 0;
+
+	virtual TSharedPtr<IPlugin> FindEnabledPluginFromPath(const FString& PluginPath) = 0;
+	virtual TSharedPtr<IPlugin> FindEnabledPluginFromDescriptor(const FPluginReferenceDescriptor& PluginDesc) = 0;
+
 	/** 
 	 * Finds all plugin descriptors underneath a given directory (recursively)
 	 * @param Directory Search folder
