Don't load localization data for unmounted explicitly loaded plugins

#rnx

diff --git a/Engine/Source/Runtime/Projects/Private/PluginManager.cpp b/Engine/Source/Runtime/Projects/Private/PluginManager.cpp
--- a/Engine/Source/Runtime/Projects/Private/PluginManager.cpp
+++ b/Engine/Source/Runtime/Projects/Private/PluginManager.cpp
@@ -2889,7 +2889,7 @@ void FPluginManager::GetLocalizationPathsForEnabledPlugins( TArray<FString>& Out
 	for (const FDiscoveredPluginMap::ElementType& PluginPair : AllPlugins)
 	{
 		const TSharedRef<FPlugin>& Plugin = DiscoveredPluginMapUtils::ResolvePluginFromMapVal(PluginPair.Value);
-		if (!Plugin->bEnabled || Plugin->GetDescriptor().LocalizationTargets.Num() == 0)
+		if (!Plugin->bEnabled || (Plugin->GetDescriptor().bExplicitlyLoaded && !Plugin->bIsMounted) || Plugin->GetDescriptor().LocalizationTargets.Num() == 0)
 		{
 			continue;
 		}
