Don't mount localization data for explicitly loaded plugins by default

This allows them to separate their mounting and content loading into two phases, such as allowing a GFP to mount the plugin when "registered" and defer mounting the localization data until "activating".

Note: Anything that was calling one of the MountExplicitlyLoadedPlugin functions to load a plugin, will now need to also call MountExplicitlyLoadedPluginLocalizationData if it wants to also mount the localization data.

diff --git a/Engine/Plugins/Runtime/GameFeatures/Source/GameFeatures/Private/GameFeaturePluginStateMachine.cpp b/Engine/Plugins/Runtime/GameFeatures/Source/GameFeatures/Private/GameFeaturePluginStateMachine.cpp
--- a/Engine/Plugins/Runtime/GameFeatures/Source/GameFeatures/Private/GameFeaturePluginStateMachine.cpp
+++ b/Engine/Plugins/Runtime/GameFeatures/Source/GameFeatures/Private/GameFeaturePluginStateMachine.cpp
@@ -2856,6 +2856,8 @@ struct FGameFeaturePluginState_Activating : public FGameFeaturePluginState
 
 		FGameFeatureActivatingContext Context;
 
+		const bool bIsLoadingLocalizationData = IPluginManager::Get().MountExplicitlyLoadedPluginLocalizationData(StateProperties.PluginName);
+
 		if (AllowIniLoading())
 		{
 			TRACE_CPUPROFILER_EVENT_SCOPE(GFP_Activating_InitIni);
@@ -2869,8 +2871,7 @@ struct FGameFeaturePluginState_Activating : public FGameFeaturePluginState
 
 		// @TODO: non-blocking wait here?
 		// If this plugin caused localization data to load, wait for that here before marking it as active
-		if (TSharedPtr<IPlugin> Plugin = IPluginManager::Get().FindPlugin(StateProperties.PluginName);
-			Plugin && Plugin->GetDescriptor().bExplicitlyLoaded && Plugin->GetDescriptor().LocalizationTargets.Num() > 0)
+		if (bIsLoadingLocalizationData)
 		{
 			TRACE_CPUPROFILER_EVENT_SCOPE(GFP_Activating_WaitForLoc);
 			FTextLocalizationManager::Get().WaitForAsyncTasks();
diff --git a/Engine/Source/Runtime/Projects/Private/PluginManager.cpp b/Engine/Source/Runtime/Projects/Private/PluginManager.cpp
--- a/Engine/Source/Runtime/Projects/Private/PluginManager.cpp
+++ b/Engine/Source/Runtime/Projects/Private/PluginManager.cpp
@@ -379,6 +379,7 @@ FPlugin::FPlugin(const FString& InFileName, const FPluginDescriptor& InDescripto
 	, Type(InType)
 	, bEnabled(false)
 	, bIsMounted(false)
+	, bIsExplicitlyLoadedLocalizationDataMounted(false)
 {
 
 }
@@ -2889,7 +2890,7 @@ void FPluginManager::GetLocalizationPathsForEnabledPlugins( TArray<FString>& Out
 	for (const FDiscoveredPluginMap::ElementType& PluginPair : AllPlugins)
 	{
 		const TSharedRef<FPlugin>& Plugin = DiscoveredPluginMapUtils::ResolvePluginFromMapVal(PluginPair.Value);
-		if (!Plugin->bEnabled || (Plugin->GetDescriptor().bExplicitlyLoaded && !Plugin->bIsMounted) || Plugin->GetDescriptor().LocalizationTargets.Num() == 0)
+		if (!Plugin->bEnabled || (Plugin->GetDescriptor().bExplicitlyLoaded && !Plugin->bIsExplicitlyLoadedLocalizationDataMounted) || Plugin->GetDescriptor().LocalizationTargets.Num() == 0)
 		{
 			continue;
 		}
@@ -3223,6 +3224,49 @@ bool FPluginManager::MountExplicitlyLoadedPlugin_FromDescriptor(const FPluginRef
 	return TryMountExplicitlyLoadedPluginVersion(DescribedPlugin);
 }
 
+bool FPluginManager::MountExplicitlyLoadedPluginLocalizationData(const FString& PluginName)
+{
+	TSharedPtr<FPlugin> Plugin = FindPluginInstance(PluginName);
+
+	TRACE_CPUPROFILER_EVENT_SCOPE(MountExplicitlyLoadedPluginLocalizationData);
+	if (!Plugin.IsValid() || !Plugin->bIsMounted)
+	{
+		// Does not exist or is not mounted
+		UE_LOG(LogPluginManager, Error, TEXT("Cannot mount plugin localization for '%s' as the plugin is unknown or not mounted. Did you forget to call MountExplicitlyLoadedPlugin?"), *PluginName);
+		return false;
+	}
+
+	if (!Plugin->Descriptor.bExplicitlyLoaded)
+	{
+		// Not supported
+		UE_LOG(LogPluginManager, Warning, TEXT("Cannot mount plugin localization for '%s' as the plugin isn't explicitly loaded."), *PluginName);
+		return false;
+	}
+
+	if (Plugin->bIsExplicitlyLoadedLocalizationDataMounted)
+	{
+		// Already loaded
+		UE_LOG(LogPluginManager, Verbose, TEXT("Ignoring request to mount plugin localization for '%s' as the localization data was already mounted."), *PluginName);
+		return false;
+	}
+
+	if (Plugin->Descriptor.LocalizationTargets.Num() == 0)
+	{
+		// Nothing to load
+		UE_LOG(LogPluginManager, Verbose, TEXT("Ignoring request to mount plugin localization for '%s' as the plugin has no localization targets defined."), *PluginName);
+		return false;
+	}
+
+	UE_LOG(LogPluginManager, Log, TEXT("Mounting plugin localization for '%s'..."), *PluginName);
+	Plugin->bIsExplicitlyLoadedLocalizationDataMounted = true;
+
+	// Notify that additional localization data should be loaded
+	TArray<FString> AdditionalLocResPaths;
+	PluginLocalizationUtils::GetLocalizationPathsForPlugin(*Plugin, AdditionalLocResPaths);
+	FTextLocalizationManager::Get().HandleLocalizationTargetsMounted(AdditionalLocResPaths);
+	return true;
+}
+
 bool FPluginManager::TryMountExplicitlyLoadedPluginVersion(TSharedRef<FPlugin>* AllPlugins_PluginPtr)
 {
 	bool bSuccess = false;
@@ -3293,7 +3337,7 @@ void FPluginManager::MountPluginFromExternalSource(const TSharedRef<FPlugin>& Pl
 	}
 
 	// Notify that additional localization data should be loaded
-	if (Plugin->Descriptor.LocalizationTargets.Num() > 0)
+	if (!Plugin->Descriptor.bExplicitlyLoaded && Plugin->Descriptor.LocalizationTargets.Num() > 0)
 	{
 		TArray<FString> AdditionalLocResPaths;
 		PluginLocalizationUtils::GetLocalizationPathsForPlugin(*Plugin, AdditionalLocResPaths);
@@ -3411,6 +3455,7 @@ bool FPluginManager::UnmountExplicitlyLoadedPlugin(const FString& PluginName, FT
 		PluginLocalizationUtils::GetLocalizationPathsForPlugin(*Plugin, AdditionalLocResPaths);
 		FTextLocalizationManager::Get().HandleLocalizationTargetsUnmounted(AdditionalLocResPaths);
 	}
+	Plugin->bIsExplicitlyLoadedLocalizationDataMounted = false;
 
 	if ((Plugin->CanContainContent() || Plugin->CanContainVerse()) && ensure(UnRegisterMountPointDelegate.IsBound()))
 	{
diff --git a/Engine/Source/Runtime/Projects/Private/PluginManager.h b/Engine/Source/Runtime/Projects/Private/PluginManager.h
--- a/Engine/Source/Runtime/Projects/Private/PluginManager.h
+++ b/Engine/Source/Runtime/Projects/Private/PluginManager.h
@@ -42,6 +42,12 @@ public:
 	 */
 	bool bIsMounted : 1;
 
+	/**
+	 * True if an explicitly loaded plugin has also mounted its localization data.
+	 * @note Unused for non-explicitly loaded plugins.
+	 */
+	bool bIsExplicitlyLoadedLocalizationDataMounted : 1;
+
 	/**
 	 * FPlugin constructor
 	 */
@@ -193,6 +199,7 @@ public:
 	virtual bool MountExplicitlyLoadedPlugin(const FString& PluginName) override;
 	virtual bool MountExplicitlyLoadedPlugin_FromFileName(const FString& PluginFileName) override;
 	virtual bool MountExplicitlyLoadedPlugin_FromDescriptor(const FPluginReferenceDescriptor& PluginDescriptor) override;
+	virtual bool MountExplicitlyLoadedPluginLocalizationData(const FString& PluginName) override;
 	virtual bool UnmountExplicitlyLoadedPlugin(const FString& PluginName, FText* OutReason) override { return UnmountExplicitlyLoadedPlugin(PluginName, OutReason, true); }
 	virtual bool UnmountExplicitlyLoadedPlugin(const FString& PluginName, FText* OutReason, bool bAllowUnloadCode) override;
 	virtual bool GetPluginDependencies(const FString& PluginName, TArray<FPluginReferenceDescriptor>& PluginDependencies) override;
diff --git a/Engine/Source/Runtime/Projects/Public/Interfaces/IPluginManager.h b/Engine/Source/Runtime/Projects/Public/Interfaces/IPluginManager.h
--- a/Engine/Source/Runtime/Projects/Public/Interfaces/IPluginManager.h
+++ b/Engine/Source/Runtime/Projects/Public/Interfaces/IPluginManager.h
@@ -463,11 +463,19 @@ public:
 	/**
 	 * Marks an explicitly loaded plugin as enabled, mounts its content and tries to load its modules.
 	 * These plugins are not loaded implicitly, but instead wait for this function to be called.
+	 * 
+	 * @note Call MountExplicitlyLoadedPluginLocalizationData if you also want to load any localization data for this plugin.
 	 */
 	virtual bool MountExplicitlyLoadedPlugin(const FString& PluginName) = 0;
 	virtual bool MountExplicitlyLoadedPlugin_FromFileName(const FString& PluginFileName) = 0;
 	virtual bool MountExplicitlyLoadedPlugin_FromDescriptor(const FPluginReferenceDescriptor& PluginDescriptor) = 0;
 
+	/**
+	 * Start loading localization data for an explicitly loaded plugin that has previously been mounted via one of the MountExplicitlyLoadedPlugin functions.
+	 * @return True if localization data started to load, or false if the plugin was missing or had no localization data to load.
+	 */
+	virtual bool MountExplicitlyLoadedPluginLocalizationData(const FString& PluginName) = 0;
+
 	/**
 	 * Marks an explicitly loaded plugin as disabled, unmounts its content (does not work on plugins with compiled modules).
 	 */
