Removing the field "Plugins" on update to force the array to be rebuilt in the same order as the plugins array rather than new items being appended to the end

diff --git a/Engine/Source/Runtime/Projects/Private/PluginDescriptor.cpp b/Engine/Source/Runtime/Projects/Private/PluginDescriptor.cpp
--- a/Engine/Source/Runtime/Projects/Private/PluginDescriptor.cpp
+++ b/Engine/Source/Runtime/Projects/Private/PluginDescriptor.cpp
@@ -543,6 +543,10 @@ void FPluginDescriptor::UpdateJson(FJsonObject& JsonObject) const
 	PreBuildSteps.UpdateJson(JsonObject, TEXT("PreBuildSteps"));
 	PostBuildSteps.UpdateJson(JsonObject, TEXT("PostBuildSteps"));
 
+	// Removing the plugins field to force the array to be rebuilt in the same order as the Plugins array otherwise 
+	// all new items are appended at the back.
+	JsonObject.RemoveField(TEXT("Plugins"));
+
 	FPluginReferenceDescriptor::UpdateArray(JsonObject, TEXT("Plugins"), Plugins);
 
 	FPluginDisallowedDescriptor::UpdateArray(JsonObject, TEXT("DisallowedPlugins"), DisallowedPlugins);
