Added plugin param for generating iads

diff --git a/Engine/Plugins/Developer/PluginUtils/Source/PluginUtils/Private/PluginUtils.cpp b/Engine/Plugins/Developer/PluginUtils/Source/PluginUtils/Private/PluginUtils.cpp
--- a/Engine/Plugins/Developer/PluginUtils/Source/PluginUtils/Private/PluginUtils.cpp
+++ b/Engine/Plugins/Developer/PluginUtils/Source/PluginUtils/Private/PluginUtils.cpp
@@ -569,6 +569,7 @@ TSharedPtr<IPlugin> FPluginUtils::CreateAndLoadNewPlugin(const FString& PluginNa
 	ExCreationParams.Descriptor.VersePath = CreationParams.VersePath;
 	ExCreationParams.Descriptor.VerseVersion = CreationParams.VerseVersion;
 	ExCreationParams.Descriptor.bEnableVerseAssetReflection = CreationParams.bEnableVerseAssetReflection;
+	ExCreationParams.Descriptor.bEnableIAD = CreationParams.bEnableIAD;
 
 	if (CreationParams.bHasModules)
 	{
diff --git a/Engine/Plugins/Developer/PluginUtils/Source/PluginUtils/Public/PluginUtils.h b/Engine/Plugins/Developer/PluginUtils/Source/PluginUtils/Public/PluginUtils.h
--- a/Engine/Plugins/Developer/PluginUtils/Source/PluginUtils/Public/PluginUtils.h
+++ b/Engine/Plugins/Developer/PluginUtils/Source/PluginUtils/Public/PluginUtils.h
@@ -98,6 +98,9 @@ public:
 		/** If to generate Verse source code definitions from assets contained in this plugin */
 		bool bEnableVerseAssetReflection = false;
 
+		/** enable iad for this plugin */
+		bool bEnableIAD;
+
 		/** Whether this plugin should be enabled/disabled by default for any project. */
 		EPluginEnabledByDefault EnabledByDefault = EPluginEnabledByDefault::Unspecified;
 
diff --git a/Engine/Source/Runtime/Projects/Private/PluginDescriptor.cpp b/Engine/Source/Runtime/Projects/Private/PluginDescriptor.cpp
--- a/Engine/Source/Runtime/Projects/Private/PluginDescriptor.cpp
+++ b/Engine/Source/Runtime/Projects/Private/PluginDescriptor.cpp
@@ -258,6 +258,8 @@ bool FPluginDescriptor::Read(const FJsonObject& Object, FText* OutFailReason /*=
 
 	Object.TryGetBoolField(TEXT("EnableVerseAssetReflection"), bEnableVerseAssetReflection);
 
+	Object.TryGetBoolField(TEXT("EnableIAD"), bEnableIAD);
+
 	bool bEnabledByDefault;
 	if(Object.TryGetBoolField(TEXT("EnabledByDefault"), bEnabledByDefault))
 	{
@@ -427,6 +429,15 @@ void FPluginDescriptor::UpdateJson(FJsonObject& JsonObject) const
 		JsonObject.RemoveField(TEXT("EnableVerseAssetReflection"));
 	}
 
+	if (bEnableIAD)
+	{
+		JsonObject.SetBoolField(TEXT("EnableIAD"), bEnableIAD);
+	}
+	else
+	{
+		JsonObject.RemoveField(TEXT("EnableIAD"));
+	}
+
 	if (EnabledByDefault != EPluginEnabledByDefault::Unspecified)
 	{
 		JsonObject.SetBoolField(TEXT("EnabledByDefault"), (EnabledByDefault == EPluginEnabledByDefault::Enabled));
diff --git a/Engine/Source/Runtime/Projects/Public/PluginDescriptor.h b/Engine/Source/Runtime/Projects/Public/PluginDescriptor.h
--- a/Engine/Source/Runtime/Projects/Public/PluginDescriptor.h
+++ b/Engine/Source/Runtime/Projects/Public/PluginDescriptor.h
@@ -104,6 +104,9 @@ struct FPluginDescriptor
 	/** If to generate Verse source code definitions from assets contained in this plugin */
 	bool bEnableVerseAssetReflection = false;
 
+	/** Use IAD for this plugin */
+	bool bEnableIAD = false;
+
 	/** Whether this plugin should be enabled by default for all projects */
 	EPluginEnabledByDefault EnabledByDefault;
 
