Adding GameTarget Allow/Deny Lists to UBT

#rnx

diff --git a/Engine/Source/Programs/UnrealBuildTool/System/ModuleDescriptor.cs b/Engine/Source/Programs/UnrealBuildTool/System/ModuleDescriptor.cs
--- a/Engine/Source/Programs/UnrealBuildTool/System/ModuleDescriptor.cs
+++ b/Engine/Source/Programs/UnrealBuildTool/System/ModuleDescriptor.cs
@@ -211,6 +211,16 @@ namespace UnrealBuildTool
 		/// List of disallowed programs
 		/// </summary>
 		public string[]? ProgramDenyList;
+		
+		/// <summary>
+		/// List of allowed game targets
+		/// </summary>
+		public string[]? GameTargetAllowList;
+
+		/// <summary>
+		/// List of disallowed game targets
+		/// </summary>
+		public string[]? GameTargetDenyList;
 
 		/// <summary>
 		/// List of additional dependencies for building this module.
@@ -331,7 +341,19 @@ namespace UnrealBuildTool
 			{
 				Module.ProgramDenyList = ProgramDenyList;
 			}
+			
+			string[]? GameTargetAllowList;
+			if (InObject.TryGetStringArrayField("GameTargetAllowList", out GameTargetAllowList))
+			{
+				Module.GameTargetAllowList = GameTargetAllowList;
+			}
 
+			string[]? GameTargetDenyList;
+			if (InObject.TryGetStringArrayField("GameTargetDenyList", out GameTargetDenyList))
+			{
+				Module.GameTargetDenyList = GameTargetDenyList;
+			}
+			
 			string[]? AdditionalDependencies;
 			if (InObject.TryGetStringArrayField("AdditionalDependencies", out AdditionalDependencies))
 			{
@@ -422,6 +444,14 @@ namespace UnrealBuildTool
 			{
 				Writer.WriteStringArrayField("ProgramDenyList", ProgramDenyList);
 			}
+			if (GameTargetAllowList != null && GameTargetAllowList.Length > 0)
+			{
+				Writer.WriteStringArrayField("GameTargetAllowList", GameTargetAllowList);
+			}
+			if (GameTargetDenyList != null && GameTargetDenyList.Length > 0)
+			{
+				Writer.WriteStringArrayField("GameTargetDenyList", GameTargetDenyList);
+			}	
 			if (AdditionalDependencies != null && AdditionalDependencies.Length > 0)
 			{
 				Writer.WriteArrayStart("AdditionalDependencies");
@@ -594,7 +624,6 @@ namespace UnrealBuildTool
 				return false;
 			}
 
-			// Special checks just for programs
 			if (TargetType == TargetType.Program)
 			{
 				// Check the program name is on the allow list. Note that this behavior is slightly different to other allow/deny checks; we will allow a module of any type if it's explicitly allowed for this program.
@@ -609,7 +638,21 @@ namespace UnrealBuildTool
 					return false;
 				}
 			}
-
+			else 
+			{
+				// Check that the TargetName is allowed
+				if (GameTargetAllowList != null && GameTargetAllowList.Length > 0 && !GameTargetAllowList.Contains(TargetName))
+				{
+					return false;
+				}
+				
+				// Check that the TargetName is not denied
+				if (GameTargetDenyList != null && GameTargetDenyList.Contains(TargetName))
+				{
+					return false;
+				}
+			}
+			
 			// Check the module is compatible with this target.
 			switch (Type)
 			{
diff --git a/Engine/Source/Runtime/Projects/Private/ModuleDescriptor.cpp b/Engine/Source/Runtime/Projects/Private/ModuleDescriptor.cpp
--- a/Engine/Source/Runtime/Projects/Private/ModuleDescriptor.cpp
+++ b/Engine/Source/Runtime/Projects/Private/ModuleDescriptor.cpp
@@ -225,7 +225,11 @@ bool FModuleDescriptor::Read(const FJsonObject& Object, FText* OutFailReason /*=
 
 	// Read the additional dependencies
 	Object.TryGetStringArrayField(TEXT("AdditionalDependencies"), AdditionalDependencies);
-
+	
+	// Read the allow/deny list for game targets
+	Object.TryGetStringArrayField(TEXT("GameTargetAllowList"), GameTargetAllowList);
+	Object.TryGetStringArrayField(TEXT("GameTargetDenyList"), GameTargetDenyList);
+	
 	return true;
 }
 
@@ -401,7 +405,34 @@ void FModuleDescriptor::UpdateJson(FJsonObject& JsonObject) const
 	{
 		JsonObject.RemoveField(TEXT("ProgramDenyList"));
 	}
+	if (GameTargetAllowList.Num() > 0)
+	{
+		TArray<TSharedPtr<FJsonValue>> GameTargetAllowListValues;
+		for (const FString& Program : ProgramAllowList)
+		{
+			GameTargetAllowListValues.Add(MakeShareable(new FJsonValueString(Program)));
+		}
+		JsonObject.SetArrayField(TEXT("GameTargetAllowList"), GameTargetAllowListValues);
+	}
+	else
+	{
+		JsonObject.RemoveField(TEXT("GameTargetAllowList"));
+	}
 
+	if (GameTargetDenyList.Num() > 0)
+	{
+		TArray<TSharedPtr<FJsonValue>> GameTargetDenyListValues;
+		for (const FString& Program : ProgramDenyList)
+		{
+			GameTargetDenyListValues.Add(MakeShareable(new FJsonValueString(Program)));
+		}
+		JsonObject.SetArrayField(TEXT("GameTargetDenyList"), GameTargetDenyListValues);
+	}
+	else
+	{
+		JsonObject.RemoveField(TEXT("GameTargetDenyList"));
+	}
+	
 	if (AdditionalDependencies.Num() > 0)
 	{
 		TArray<TSharedPtr<FJsonValue>> AdditionalDependencyValues;
@@ -513,7 +544,19 @@ bool FModuleDescriptor::IsCompiledInConfiguration(const FString& Platform, EBuil
 			return false;
 		}
 	}
-
+	else
+	{
+		if(GameTargetAllowList.Num() > 0 && !GameTargetAllowList.Contains(TargetName))
+		{
+			return false;
+		}
+				
+		if(GameTargetDenyList.Contains(TargetName))
+		{
+			return false;
+		}
+	}
+	
 	// Check the module is compatible with this target.
 	switch (Type)
 	{
@@ -560,12 +603,13 @@ bool FModuleDescriptor::IsLoadedInCurrentConfiguration() const
 		return false;
 	}
 
+	
 	// Always respect the allow/deny lists for program targets
 	EBuildTargetType TargetType = FApp::GetBuildTargetType();
 	if(TargetType == EBuildTargetType::Program)
 	{
 		const FString TargetName = UE_APP_NAME;
-
+			
 		// Check the program name is allowed. Note that this behavior is slightly different to other allow/deny list checks; we will allow a module of any type if it's explicitly allowed for this program.
 		if(ProgramAllowList.Num() > 0)
 		{
@@ -578,6 +622,20 @@ bool FModuleDescriptor::IsLoadedInCurrentConfiguration() const
 			return false;
 		}
 	}
+	else
+	{
+		const FString TargetName = UE_APP_NAME;
+		
+		if(GameTargetAllowList.Num() > 0 && !GameTargetAllowList.Contains(TargetName))
+		{
+			return false;
+		}
+				
+		if(GameTargetDenyList.Contains(TargetName))
+		{
+			return false;
+		}
+	}
 
 	// Check that the runtime environment allows it to be loaded
 	switch (Type)
diff --git a/Engine/Source/Runtime/Projects/Public/ModuleDescriptor.h b/Engine/Source/Runtime/Projects/Public/ModuleDescriptor.h
--- a/Engine/Source/Runtime/Projects/Public/ModuleDescriptor.h
+++ b/Engine/Source/Runtime/Projects/Public/ModuleDescriptor.h
@@ -186,6 +186,12 @@ struct FModuleDescriptor
 	/** List of disallowed programs */
 	TArray<FString> ProgramDenyList;
 
+	/** List of allowed game targets */
+	TArray<FString> GameTargetAllowList;
+
+	/** List of disallowed game targets */
+	TArray<FString> GameTargetDenyList;
+	
 	/** List of additional dependencies for building this module. */
 	TArray<FString> AdditionalDependencies;
 
