Enabling the runtime root plugin optional references code so that not enabled plugins are ignored with snapshot use. Non shipping builds only. This is required as the staging step which performs the stripping is not executed in this workflow

diff --git a/Engine/Source/Programs/UnrealBuildTool/Configuration/UEBuildTarget.cs b/Engine/Source/Programs/UnrealBuildTool/Configuration/UEBuildTarget.cs
--- a/Engine/Source/Programs/UnrealBuildTool/Configuration/UEBuildTarget.cs
+++ b/Engine/Source/Programs/UnrealBuildTool/Configuration/UEBuildTarget.cs
@@ -4427,7 +4427,7 @@ namespace UnrealBuildTool
 		/// <param name="moduleRules">the module's rules</param>
 		private void ModifyModuleRulesForTarget(string moduleName, ModuleRules moduleRules)
 		{
-			if (moduleName == "Projects" && ((Rules.Type == TargetType.Editor || Rules.Type == TargetType.Program)))
+			if (moduleName == "Projects" && (Rules.Configuration != UnrealTargetConfiguration.Shipping))
 			{
 				// Monolithic and non-shared environment builds compile in the build plugins. Non-monolithic editor builds save them in the BuildPlugins receipt to avoid invalidating the shared build environment.
 				// See Projects.Build.cs
diff --git a/Engine/Source/Runtime/Projects/Private/PluginManager.cpp b/Engine/Source/Runtime/Projects/Private/PluginManager.cpp
--- a/Engine/Source/Runtime/Projects/Private/PluginManager.cpp
+++ b/Engine/Source/Runtime/Projects/Private/PluginManager.cpp
@@ -1437,7 +1437,7 @@ bool FPluginManager::ConfigureEnabledPlugins()
 		};
 #endif // READ_TARGET_ENABLED_PLUGINS_FROM_RECEIPT
 
-#if !WITH_EDITOR
+#if UE_BUILD_SHIPPING
 		// const to ensure it to stays empty
 		const TSet<FString> AllowedOptionalDependencies;
 #else
@@ -1474,7 +1474,7 @@ bool FPluginManager::ConfigureEnabledPlugins()
 			AllowedOptionalDependencies.Append({ UBT_TARGET_BUILD_PLUGINS });
 		}
 #endif // READ_TARGET_ENABLED_PLUGINS_FROM_RECEIPT
-#endif // !WITH_EDITOR
+#endif // UE_BUILD_SHIPPING
 
 		// Check which plugins have been enabled or excluded via the command line
 		{
@@ -1538,9 +1538,9 @@ bool FPluginManager::ConfigureEnabledPlugins()
 			}
 			if (ExtraPluginsToEnable.Num() > 0)
 			{
-#if WITH_EDITOR
+#if !UE_BUILD_SHIPPING
 				AllowedOptionalDependencies.Append(ExtraPluginsToEnable);
-#endif // WITH_EDITOR
+#endif
 
 				auto IsRestrictedPlugin = [this](const FString& PluginName)
 				{
@@ -2543,7 +2543,7 @@ bool FPluginManager::ConfigureEnabledPluginForTarget(const FPluginReferenceDescr
 				return false;
 			}
 
-#if WITH_EDITOR
+#if !UE_BUILD_SHIPPING
 			// Allowed optional plugins are compiled enabled or enabled via the commandline. Ignore those that are not.
 			if (NextReference.bOptional && !AllowedOptionalDependencies.IsEmpty() && !AllowedOptionalDependencies.Contains(NextReference.Name))
 			{
