Ref-count the mounting of localization data for explicitly mounted plugins

This allows different systems that need plugin localization to call MountExplicitlyLoadedPluginLocalizationData and UnmountExplicitlyLoadedPluginLocalizationData (in pairs), without needing to understand any other system that may have mounted localization for that plugin, and avoids the unmount calls from different systems fighting with each other.

diff --git a/Engine/Source/Runtime/Projects/Private/PluginManager.cpp b/Engine/Source/Runtime/Projects/Private/PluginManager.cpp
--- a/Engine/Source/Runtime/Projects/Private/PluginManager.cpp
+++ b/Engine/Source/Runtime/Projects/Private/PluginManager.cpp
@@ -379,7 +379,6 @@ FPlugin::FPlugin(const FString& InFileName, const FPluginDescriptor& InDescripto
 	, Type(InType)
 	, bEnabled(false)
 	, bIsMounted(false)
-	, bIsExplicitlyLoadedLocalizationDataMounted(false)
 {
 
 }
@@ -2890,7 +2889,7 @@ void FPluginManager::GetLocalizationPathsForEnabledPlugins( TArray<FString>& Out
 	for (const FDiscoveredPluginMap::ElementType& PluginPair : AllPlugins)
 	{
 		const TSharedRef<FPlugin>& Plugin = DiscoveredPluginMapUtils::ResolvePluginFromMapVal(PluginPair.Value);
-		if (!Plugin->bEnabled || (Plugin->GetDescriptor().bExplicitlyLoaded && !Plugin->bIsExplicitlyLoadedLocalizationDataMounted) || Plugin->GetDescriptor().LocalizationTargets.Num() == 0)
+		if (!Plugin->bEnabled || (Plugin->GetDescriptor().bExplicitlyLoaded && Plugin->ExplicitlyLoadedLocalizationDataMountedRefCount == 0) || Plugin->GetDescriptor().LocalizationTargets.Num() == 0)
 		{
 			continue;
 		}
@@ -3243,9 +3242,10 @@ bool FPluginManager::MountExplicitlyLoadedPluginLocalizationData(const FString&
 		return false;
 	}
 
-	if (Plugin->bIsExplicitlyLoadedLocalizationDataMounted)
+	if (Plugin->ExplicitlyLoadedLocalizationDataMountedRefCount > 0)
 	{
 		// Already loaded
+		++Plugin->ExplicitlyLoadedLocalizationDataMountedRefCount;
 		UE_LOG(LogPluginManager, Verbose, TEXT("Ignoring request to mount plugin localization for '%s' as the localization data was already mounted."), *PluginName);
 		return false;
 	}
@@ -3258,7 +3258,7 @@ bool FPluginManager::MountExplicitlyLoadedPluginLocalizationData(const FString&
 	}
 
 	UE_LOG(LogPluginManager, Log, TEXT("Mounting plugin localization for '%s'..."), *PluginName);
-	Plugin->bIsExplicitlyLoadedLocalizationDataMounted = true;
+	Plugin->ExplicitlyLoadedLocalizationDataMountedRefCount = 1;
 
 	// Notify that additional localization data should be loaded
 	TArray<FString> AdditionalLocResPaths;
@@ -3286,7 +3286,15 @@ bool FPluginManager::UnmountExplicitlyLoadedPluginLocalizationData(const FString
 		return false;
 	}
 
-	if (!Plugin->bIsExplicitlyLoadedLocalizationDataMounted)
+	if (Plugin->ExplicitlyLoadedLocalizationDataMountedRefCount > 1)
+	{
+		// Not possible to unload yet
+		UE_LOG(LogPluginManager, Verbose, TEXT("Ignoring request to unmount plugin localization for '%s' as the localization data ref-count was %d."), *PluginName, Plugin->ExplicitlyLoadedLocalizationDataMountedRefCount);
+		--Plugin->ExplicitlyLoadedLocalizationDataMountedRefCount;
+		return false;
+	}
+
+	if (Plugin->ExplicitlyLoadedLocalizationDataMountedRefCount == 0)
 	{
 		// Already unloaded
 		UE_LOG(LogPluginManager, Verbose, TEXT("Ignoring request to unmount plugin localization for '%s' as the localization data was not mounted."), *PluginName);
@@ -3308,7 +3316,7 @@ bool FPluginManager::UnmountExplicitlyLoadedPluginLocalizationData(const FString
 	}
 
 	UE_LOG(LogPluginManager, Log, TEXT("Unmounting plugin localization for '%s'..."), *PluginName);
-	Plugin->bIsExplicitlyLoadedLocalizationDataMounted = false;
+	Plugin->ExplicitlyLoadedLocalizationDataMountedRefCount = 0;
 
 	// Notify that additional localization data should be unloaded
 	TArray<FString> AdditionalLocResPaths;
@@ -3499,9 +3507,9 @@ bool FPluginManager::UnmountExplicitlyLoadedPlugin(const FString& PluginName, FT
 	}
 
 	// Notify that additional localization data should be unloaded
-	if (Plugin->bIsExplicitlyLoadedLocalizationDataMounted && Plugin->Descriptor.LocalizationTargets.Num() > 0 && !IsEngineExitRequested())
+	if (Plugin->ExplicitlyLoadedLocalizationDataMountedRefCount > 0 && Plugin->Descriptor.LocalizationTargets.Num() > 0 && !IsEngineExitRequested())
 	{
-		Plugin->bIsExplicitlyLoadedLocalizationDataMounted = false;
+		Plugin->ExplicitlyLoadedLocalizationDataMountedRefCount = 0;
 
 		TArray<FString> AdditionalLocResPaths;
 		PluginLocalizationUtils::GetLocalizationPathsForPlugin(*Plugin, AdditionalLocResPaths);
diff --git a/Engine/Source/Runtime/Projects/Private/PluginManager.h b/Engine/Source/Runtime/Projects/Private/PluginManager.h
--- a/Engine/Source/Runtime/Projects/Private/PluginManager.h
+++ b/Engine/Source/Runtime/Projects/Private/PluginManager.h
@@ -43,10 +43,10 @@ public:
 	bool bIsMounted : 1;
 
 	/**
-	 * True if an explicitly loaded plugin has also mounted its localization data.
+	 * >0 if an explicitly loaded plugin has also mounted its localization data.
 	 * @note Unused for non-explicitly loaded plugins.
 	 */
-	bool bIsExplicitlyLoadedLocalizationDataMounted : 1;
+	uint8 ExplicitlyLoadedLocalizationDataMountedRefCount = 0;
 
 	/**
 	 * FPlugin constructor
diff --git a/Engine/Source/Runtime/Projects/Public/Interfaces/IPluginManager.h b/Engine/Source/Runtime/Projects/Public/Interfaces/IPluginManager.h
--- a/Engine/Source/Runtime/Projects/Public/Interfaces/IPluginManager.h
+++ b/Engine/Source/Runtime/Projects/Public/Interfaces/IPluginManager.h
@@ -472,13 +472,15 @@ public:
 
 	/**
 	 * Start loading localization data for an explicitly loaded plugin that has previously been mounted via one of the MountExplicitlyLoadedPlugin functions.
+	 * @note Localization data is ref-counted, so this should be paired with a call to UnmountExplicitlyLoadedPluginLocalizationData.
 	 * @return True if localization data started to load, or false if the plugin was missing or had no localization data to load.
 	 */
 	virtual bool MountExplicitlyLoadedPluginLocalizationData(const FString& PluginName) = 0;
 
 	/**
 	 * Start unloading localization data for an explicitly loaded plugin that had its localization data mounted via MountExplicitlyLoadedPluginLocalizationData.
-	 * @note Localization data is also automatically unloaded when calling UnmountExplicitlyLoadedPlugin.
+	 * @note Localization data is ref-counted, so this should be paired with a call to MountExplicitlyLoadedPluginLocalizationData.
+	 * @note Localization data is also automatically (and regardless of ref-count) unloaded when calling UnmountExplicitlyLoadedPlugin.
 	 * @return True if localization data started to unload, or false if the plugin was missing or had no localization data to unload.
 	 */
 	virtual bool UnmountExplicitlyLoadedPluginLocalizationData(const FString& PluginName) = 0;
