Avoid localization data failing to load for loaded/active GFPs if GameFeaturePlugin.DeferLocalizationDataLoad is hotfixed to false after the GFPs have already mounted

Since the mount/unmount localization calls are now ref-counted, we can just always call them during GFP loading/unloading as they'll be a no-op if the localization data was already loaded by mounting (or anything else).

Also improved the logging when the localization data skips being loaded/unloaded due to an existing reference, so that we can see the ref-count in the log.

#rnx

diff --git a/Engine/Plugins/Runtime/GameFeatures/Source/GameFeatures/Private/GameFeaturePluginStateMachine.cpp b/Engine/Plugins/Runtime/GameFeatures/Source/GameFeatures/Private/GameFeaturePluginStateMachine.cpp
--- a/Engine/Plugins/Runtime/GameFeatures/Source/GameFeatures/Private/GameFeaturePluginStateMachine.cpp
+++ b/Engine/Plugins/Runtime/GameFeatures/Source/GameFeatures/Private/GameFeaturePluginStateMachine.cpp
@@ -2755,10 +2755,7 @@ struct FGameFeaturePluginState_Unloading : public FGameFeaturePluginState
 
 	virtual void BeginState() override
 	{
-		if (UE::GameFeatures::ShouldDeferLocalizationDataLoad())
-		{
-			IPluginManager::Get().UnmountExplicitlyLoadedPluginLocalizationData(StateProperties.PluginName);
-		}
+		IPluginManager::Get().UnmountExplicitlyLoadedPluginLocalizationData(StateProperties.PluginName);
 	}
 
 	virtual void UpdateState(FGameFeaturePluginStateStatus& StateStatus) override
@@ -2807,10 +2804,7 @@ struct FGameFeaturePluginState_Loading : public FGameFeaturePluginState
 		TRACE_CPUPROFILER_EVENT_SCOPE(GFP_Loading_Begin);
 		check(StateProperties.GameFeatureData);
 
-		if (UE::GameFeatures::ShouldDeferLocalizationDataLoad())
-		{
-			IPluginManager::Get().MountExplicitlyLoadedPluginLocalizationData(StateProperties.PluginName);
-		}
+		IPluginManager::Get().MountExplicitlyLoadedPluginLocalizationData(StateProperties.PluginName);
 
 		BundleHandle = LoadGameFeatureBundles(StateProperties.GameFeatureData);
 		if (BundleHandle)
diff --git a/Engine/Source/Runtime/Projects/Private/PluginManager.cpp b/Engine/Source/Runtime/Projects/Private/PluginManager.cpp
--- a/Engine/Source/Runtime/Projects/Private/PluginManager.cpp
+++ b/Engine/Source/Runtime/Projects/Private/PluginManager.cpp
@@ -3245,8 +3245,8 @@ bool FPluginManager::MountExplicitlyLoadedPluginLocalizationData(const FString&
 	if (Plugin->ExplicitlyLoadedLocalizationDataMountedRefCount > 0)
 	{
 		// Already loaded
+		UE_LOG(LogPluginManager, Verbose, TEXT("Ignoring request to mount plugin localization for '%s' as it was already referenced (ref-count: %d -> %d)."), *PluginName, Plugin->ExplicitlyLoadedLocalizationDataMountedRefCount, Plugin->ExplicitlyLoadedLocalizationDataMountedRefCount + 1);
 		++Plugin->ExplicitlyLoadedLocalizationDataMountedRefCount;
-		UE_LOG(LogPluginManager, Verbose, TEXT("Ignoring request to mount plugin localization for '%s' as the localization data was already mounted."), *PluginName);
 		return false;
 	}
 
@@ -3289,7 +3289,7 @@ bool FPluginManager::UnmountExplicitlyLoadedPluginLocalizationData(const FString
 	if (Plugin->ExplicitlyLoadedLocalizationDataMountedRefCount > 1)
 	{
 		// Not possible to unload yet
-		UE_LOG(LogPluginManager, Verbose, TEXT("Ignoring request to unmount plugin localization for '%s' as the localization data ref-count was %d."), *PluginName, Plugin->ExplicitlyLoadedLocalizationDataMountedRefCount);
+		UE_LOG(LogPluginManager, Verbose, TEXT("Ignoring request to unmount plugin localization for '%s' as it is still referenced (ref-count: %d -> %d)."), *PluginName, Plugin->ExplicitlyLoadedLocalizationDataMountedRefCount, Plugin->ExplicitlyLoadedLocalizationDataMountedRefCount - 1);
 		--Plugin->ExplicitlyLoadedLocalizationDataMountedRefCount;
 		return false;
 	}
