Adding an optional sort delegate to allow custom sorting of the updated json array before it's set. The default behavor adds new elements to the end of the json array which may not match the source input array. Removing the clearing of "Plugins" so that existing data is preserved again.

diff --git a/Engine/Source/Runtime/Json/Public/JsonUtils/JsonObjectArrayUpdater.h b/Engine/Source/Runtime/Json/Public/JsonUtils/JsonObjectArrayUpdater.h
--- a/Engine/Source/Runtime/Json/Public/JsonUtils/JsonObjectArrayUpdater.h
+++ b/Engine/Source/Runtime/Json/Public/JsonUtils/JsonObjectArrayUpdater.h
@@ -25,7 +25,9 @@ struct FJsonObjectArrayUpdater
 
 	DECLARE_DELEGATE_TwoParams(FUpdateJsonObject, const ElementType&, FJsonObject&);
 
-	static void Execute(FJsonObject& JsonObject, const FString& ArrayName, const TArray<ElementType>& SourceArray, FGetElementKey GetElementKey, FTryGetJsonObjectKey TryGetJsonObjectKey, FUpdateJsonObject UpdateJsonObject)
+	DECLARE_DELEGATE_OneParam(FSortArray, TArray<TSharedPtr<FJsonValue>>&);
+
+	static void Execute(FJsonObject& JsonObject, const FString& ArrayName, const TArray<ElementType>& SourceArray, FGetElementKey GetElementKey, FTryGetJsonObjectKey TryGetJsonObjectKey, FUpdateJsonObject UpdateJsonObject, FSortArray SortJsonArray = FSortArray())
 	{
 		if (SourceArray.Num() > 0)
 		{
@@ -80,6 +82,10 @@ struct FJsonObjectArrayUpdater
 				}
 			}
 
+			// Optionally sort new json array before setting incase the user doesn't want the default
+			// appending of new items at the end.
+			SortJsonArray.ExecuteIfBound(NewJsonValues);
+
 			// Set the new content of the json array
 			JsonObject.SetArrayField(ArrayName, NewJsonValues);
 		}
diff --git a/Engine/Source/Runtime/Projects/Private/PluginDescriptor.cpp b/Engine/Source/Runtime/Projects/Private/PluginDescriptor.cpp
--- a/Engine/Source/Runtime/Projects/Private/PluginDescriptor.cpp
+++ b/Engine/Source/Runtime/Projects/Private/PluginDescriptor.cpp
@@ -557,10 +557,6 @@ void FPluginDescriptor::UpdateJson(FJsonObject& JsonObject) const
 	PreBuildSteps.UpdateJson(JsonObject, TEXT("PreBuildSteps"));
 	PostBuildSteps.UpdateJson(JsonObject, TEXT("PostBuildSteps"));
 
-	// Removing the plugins field to force the array to be rebuilt in the same order as the Plugins array otherwise 
-	// all new items are appended at the back.
-	JsonObject.RemoveField(TEXT("Plugins"));
-
 	FPluginReferenceDescriptor::UpdateArray(JsonObject, TEXT("Plugins"), Plugins);
 
 	FPluginDisallowedDescriptor::UpdateArray(JsonObject, TEXT("DisallowedPlugins"), DisallowedPlugins);
diff --git a/Engine/Source/Runtime/Projects/Private/PluginReferenceDescriptor.cpp b/Engine/Source/Runtime/Projects/Private/PluginReferenceDescriptor.cpp
--- a/Engine/Source/Runtime/Projects/Private/PluginReferenceDescriptor.cpp
+++ b/Engine/Source/Runtime/Projects/Private/PluginReferenceDescriptor.cpp
@@ -435,7 +435,31 @@ void FPluginReferenceDescriptor::UpdateArray(FJsonObject& JsonObject, const TCHA
 		JsonObject, ArrayName, Plugins, 
 		FPluginRefJsonArrayUpdater::FGetElementKey::CreateStatic(PluginReferenceDescriptor::GetPluginRefKey),
 		FPluginRefJsonArrayUpdater::FTryGetJsonObjectKey::CreateStatic(PluginReferenceDescriptor::TryGetPluginRefJsonObjectKey),
-		FPluginRefJsonArrayUpdater::FUpdateJsonObject::CreateStatic(PluginReferenceDescriptor::UpdatePluginRefJsonObject));
+		FPluginRefJsonArrayUpdater::FUpdateJsonObject::CreateStatic(PluginReferenceDescriptor::UpdatePluginRefJsonObject),
+		FPluginRefJsonArrayUpdater::FSortArray::CreateLambda([&Plugins](TArray<TSharedPtr<FJsonValue>>& NewJsonValues) {
+			// Sort the json array to match the same order as the plugin array. Without the sort, new entries are appended at the end
+			for (int32 StartIndex = 0; StartIndex < Plugins.Num(); ++StartIndex)
+			{
+				const FString PluginRefKey = PluginReferenceDescriptor::GetPluginRefKey(Plugins[StartIndex]);
+				for (int32 Index = StartIndex; Index < NewJsonValues.Num(); ++Index)
+				{
+					const TSharedPtr<FJsonObject>* ExistingJsonValueAsObject;
+					if (NewJsonValues[Index]->TryGetObject(ExistingJsonValueAsObject))
+					{
+						FString ElementKey;
+						if (PluginReferenceDescriptor::TryGetPluginRefJsonObjectKey(**ExistingJsonValueAsObject, ElementKey))
+						{
+							if (ElementKey == PluginRefKey)
+							{
+								NewJsonValues.Swap(StartIndex, Index);
+								break;
+							}
+						}
+					}
+				}
+			}
+		})
+	);
 }
 
 #if WITH_EDITOR
