UnrealBuildTool: Introduce TargetRules.bAllowEnginePluginsEnabledByDefault, analogous to (the inverse of) "DisableEnginePluginsByDefault" in the project descriptor.

This is passed through to the target receipt, and respected in FPluginManager.

diff --git a/Engine/Source/Programs/UnrealBuildTool/Configuration/ReadOnlyTargetRules.cs b/Engine/Source/Programs/UnrealBuildTool/Configuration/ReadOnlyTargetRules.cs
--- a/Engine/Source/Programs/UnrealBuildTool/Configuration/ReadOnlyTargetRules.cs
+++ b/Engine/Source/Programs/UnrealBuildTool/Configuration/ReadOnlyTargetRules.cs
@@ -108,6 +108,8 @@ namespace UnrealBuildTool
 
 		public IReadOnlyList<string> OptionalPlugins => Inner.OptionalPlugins;
 
+		public bool bAllowEnginePluginsEnabledByDefault => Inner.bAllowEnginePluginsEnabledByDefault;
+
 		public WarningLevel DisablePluginsConflictWarningLevel => Inner.DisablePluginsConflictWarningLevel;
 
 		public IReadOnlyList<string> BuildPlugins => Inner.BuildPlugins;
diff --git a/Engine/Source/Programs/UnrealBuildTool/Configuration/TargetRules.cs b/Engine/Source/Programs/UnrealBuildTool/Configuration/TargetRules.cs
--- a/Engine/Source/Programs/UnrealBuildTool/Configuration/TargetRules.cs
+++ b/Engine/Source/Programs/UnrealBuildTool/Configuration/TargetRules.cs
@@ -786,6 +786,11 @@ namespace UnrealBuildTool
 		[CommandLine("-OptionalPlugins=", ListSeparator = '+')]
 		public List<string> OptionalPlugins = new();
 
+		/// <summary>
+		/// If false, suppress loading "enabled by default" engine plugins not explicitly enabled by the project or target files.
+		/// </summary>
+		public bool bAllowEnginePluginsEnabledByDefault { get; set; } = true;
+
 		/// <summary>
 		/// How to treat conflicts when a disabled plugin is being enabled by another plugin referencing it
 		/// </summary>
diff --git a/Engine/Source/Programs/UnrealBuildTool/Configuration/UEBuildTarget.cs b/Engine/Source/Programs/UnrealBuildTool/Configuration/UEBuildTarget.cs
--- a/Engine/Source/Programs/UnrealBuildTool/Configuration/UEBuildTarget.cs
+++ b/Engine/Source/Programs/UnrealBuildTool/Configuration/UEBuildTarget.cs
@@ -1527,7 +1527,13 @@ namespace UnrealBuildTool
 			if (ProjectDescriptor?.Modules?.Any(x => x.Name == module.Name) == true)
 			{
 				string projectName = ProjectFile?.FullName ?? TargetRulesFile.GetFileName();
-				bool bAllowEnginePluginsEnabledByDefault = (!ProjectDescriptor?.DisableEnginePluginsByDefault) ?? false;
+
+				// EffectiveAllow = TargetRulesAllow && ProjectAllow;
+				bool bAllowEnginePluginsEnabledByDefault = Rules.bAllowEnginePluginsEnabledByDefault;
+				if (ProjectDescriptor?.DisableEnginePluginsByDefault ?? false)
+				{
+					bAllowEnginePluginsEnabledByDefault = false;
+				}
 
 				// Only warn on modules that can be compiled
 				foreach (UEBuildModuleCPP dependencyModule in referencedModules.Value.OfType<UEBuildModuleCPP>())
@@ -2378,6 +2384,11 @@ namespace UnrealBuildTool
 				Receipt.BuildPlugins.Add(Plugin.Name);
 			}
 
+			if (!Rules.bAllowEnginePluginsEnabledByDefault)
+			{
+				Receipt.AdditionalProperties.Add(new ReceiptProperty("bAllowEnginePluginsEnabledByDefault", "false"));
+			}
+
 			// Find all the modules which are part of this target
 			HashSet<UEBuildModule> UniqueLinkedModules = new HashSet<UEBuildModule>();
 			foreach (UEBuildBinary Binary in Binaries)
@@ -4596,12 +4607,17 @@ namespace UnrealBuildTool
 				}
 			}
 
-			bool bAllowEnginePluginsEnabledByDefault = true;
+			bool bAllowEnginePluginsEnabledByDefault = Rules.bAllowEnginePluginsEnabledByDefault;
 
 			// Find a map of plugins which are explicitly referenced in the project file
 			if (ProjectDescriptor != null)
 			{
-				bAllowEnginePluginsEnabledByDefault = !ProjectDescriptor.DisableEnginePluginsByDefault;
+				// EffectiveAllow = TargetRulesAllow && ProjectAllow;
+				if (ProjectDescriptor.DisableEnginePluginsByDefault)
+				{
+					bAllowEnginePluginsEnabledByDefault = false;
+				}
+
 				if (ProjectDescriptor.Plugins != null)
 				{
 					string ProjectReferenceChain = ProjectFile!.GetFileName();
diff --git a/Engine/Source/Runtime/Projects/Private/PluginManager.cpp b/Engine/Source/Runtime/Projects/Private/PluginManager.cpp
--- a/Engine/Source/Runtime/Projects/Private/PluginManager.cpp
+++ b/Engine/Source/Runtime/Projects/Private/PluginManager.cpp
@@ -27,6 +27,7 @@
 #include "Async/ParallelFor.h"
 #include "Misc/ScopeRWLock.h"
 #include "Algo/Accumulate.h"
+#include "Algo/Find.h"
 #include "Algo/Reverse.h"
 #include "Containers/VersePath.h"
 #include "Internationalization/TextLocalizationManager.h"
@@ -1591,9 +1592,14 @@ bool FPluginManager::ConfigureEnabledPlugins()
 		{
 			SCOPED_BOOT_TIMING("EnginePlugins");
 
+			// May be overridden later by receipt or project descriptor.
+			bool bAllowEnginePluginsEnabledByDefault = true;
+
 #if READ_TARGET_ENABLED_PLUGINS_FROM_RECEIPT
 			// Configure the plugins that were enabled or disabled from the target file using the target receipt file
-			auto ConfigurePluginsFromFirstMatchingTargetFile = [this, &FindFirstMatchingTargetFile, &ConfiguredPluginNames, &EnabledPlugins, &AllowedOptionalDependencies](const TCHAR* BaseDir, bool& bOutError) -> bool
+			auto ConfigurePluginsFromFirstMatchingTargetFile = [this, &FindFirstMatchingTargetFile, &ConfiguredPluginNames,
+				&EnabledPlugins, &AllowedOptionalDependencies, &bAllowEnginePluginsEnabledByDefault]
+				(const TCHAR* BaseDir, bool& bOutError) -> bool
 			{
 				const FString ReceiptWildcard = FTargetReceipt::GetDefaultPath(BaseDir, TEXT("*"), FPlatformProcess::GetBinariesSubdirectory(), FApp::GetBuildConfiguration(), nullptr);
 				FString SourceDescription = FString::Printf(TEXT("Receipt files %s"), *ReceiptWildcard);
@@ -1601,6 +1607,13 @@ bool FPluginManager::ConfigureEnabledPlugins()
 				TUniquePtr<FTargetReceipt> Receipt = FindFirstMatchingTargetFile(ReceiptWildcard);
 				if (Receipt.IsValid())
 				{
+					FReceiptProperty* AllowDefaultProperty = Algo::FindByPredicate(Receipt->AdditionalProperties,
+						[](const FReceiptProperty& InProp) -> bool { return InProp.Name == TEXT("bAllowEnginePluginsEnabledByDefault"); });
+					if (AllowDefaultProperty && !FCString::Stricmp(*AllowDefaultProperty->Value, TEXT("false")))
+					{
+						bAllowEnginePluginsEnabledByDefault = false;
+					}
+
 					for (const TPair<FString, bool>& Pair : Receipt->PluginNameToEnabledState)
 					{
 						const FString& PluginName = Pair.Key;
@@ -1692,7 +1705,6 @@ bool FPluginManager::ConfigureEnabledPlugins()
 				return true;
 			};
 
-			bool bAllowEnginePluginsEnabledByDefault = true;
 			// Find all the plugin references in the project file
 			const FProjectDescriptor* ProjectDescriptor = IProjectManager::Get().GetCurrentProject();
 			{
