Fix issue where GFS plugins or dependencies of root plugins might become unmounted and not enabled. This would prevent them from being mounted again by requiring that the plugins be enabled when tied to a root.
Add way to know the current enabled plugins for the target. This allows the application to only load the modules and plugins.

diff --git a/Engine/Source/Runtime/Projects/Private/PluginManager.cpp b/Engine/Source/Runtime/Projects/Private/PluginManager.cpp
--- a/Engine/Source/Runtime/Projects/Private/PluginManager.cpp
+++ b/Engine/Source/Runtime/Projects/Private/PluginManager.cpp
@@ -1851,6 +1851,15 @@ bool FPluginManager::ConfigureEnabledPlugins()
 			}
 		}
 
+		PluginsEnabledForTarget.Reserve(EnabledPlugins.Num());
+		for (const TPair<FString, FPlugin*>& Pair : EnabledPlugins)
+		{
+			if (Pair.Value->bEnabled)
+			{
+				PluginsEnabledForTarget.Add(Pair.Key);
+			}
+		}
+
 		// If we made it here, we have all the required plugins
 		bHaveAllRequiredPlugins = true;
 
@@ -2858,6 +2867,11 @@ bool FPluginManager::LoadModulesForEnabledPlugins( const ELoadingPhase::Type Loa
 			UE_SCOPED_ENGINE_ACTIVITY("Preloading all plugin binaries");
 			for (const FDiscoveredPluginMap::ElementType& PluginPair : AllPlugins)
 			{
+				if (!CanEnablePluginInCurrentTarget(PluginPair.Key))
+				{
+					UE_LOG(LogPluginManager, Display, TEXT("Skipping loading modules for %s. Plugin is not enabled for current target."), *PluginPair.Key);
+					continue;
+				}
 				for (auto& ModuleName : DiscoveredPluginMapUtils::ResolvePluginFromMapVal(PluginPair.Value)->Descriptor.Modules)
 				{
 					if (ModuleName.IsCompiledInCurrentConfiguration())
@@ -3053,6 +3067,16 @@ TSharedPtr<IPlugin> FPluginManager::FindEnabledPluginFromDescriptor(const FPlugi
 	}
 }
 
+bool FPluginManager::CanEnablePluginInCurrentTarget(const FStringView Name)
+{
+	if (ConfigureEnabledPlugins())
+	{
+		const uint32 NameHash = GetTypeHash(Name);
+		return PluginsEnabledForTarget.ContainsByHash(NameHash, Name);
+	}
+	return false;
+}
+
 TArray<TSharedRef<IPlugin>> FPluginManager::GetEnabledPlugins()
 {
 	TArray<TSharedRef<IPlugin>> Plugins;
@@ -3378,6 +3402,7 @@ void FPluginManager::MountPluginFromExternalSource(const TSharedRef<FPlugin>& Pl
 
 	// Mark the plugin as enabled
 	Plugin->bEnabled = true;
+	PluginsEnabledForTarget.Add(Plugin->GetName());
 
 	// Mount the plugin content directory
 	const bool bHasContentOrVerse = (Plugin->CanContainContent() || Plugin->CanContainVerse()) && ensure(RegisterMountPointDelegate.IsBound());
diff --git a/Engine/Source/Runtime/Projects/Private/PluginManager.h b/Engine/Source/Runtime/Projects/Private/PluginManager.h
--- a/Engine/Source/Runtime/Projects/Private/PluginManager.h
+++ b/Engine/Source/Runtime/Projects/Private/PluginManager.h
@@ -174,6 +174,13 @@ public:
 	virtual TSharedPtr<IPlugin> FindEnabledPluginFromPath(const FString& PluginPath) override;
 	virtual TSharedPtr<IPlugin> FindEnabledPluginFromDescriptor(const FPluginReferenceDescriptor& PluginDesc) override;
 
+	virtual bool CanEnablePluginInCurrentTarget(const FStringView Name) override;
+	virtual bool CanEnablePluginInCurrentTarget(const ANSICHAR* Name) override
+	{
+		FString NameString(Name);
+		return CanEnablePluginInCurrentTarget(FStringView(NameString));
+	}
+
 	virtual void FindPluginsUnderDirectory(const FString& Directory, TArray<FString>& OutPluginFilePaths) override;
 
 	virtual TArray<TSharedRef<IPlugin>> GetEnabledPlugins() override;
@@ -313,6 +320,9 @@ private:
 	/** Plugins that need to be configured to see if they should be enabled */
 	TSet<FString> PluginsToConfigure;
 
+	/** Set of plugins enabled for the current target. Plugins in list might not be currently enabled. */
+	TSet<FString> PluginsEnabledForTarget;
+
 #if WITH_EDITOR
 	/** Names of built-in plugins */
 	TSet<FString> BuiltInPluginNames;
diff --git a/Engine/Source/Runtime/Projects/Public/Interfaces/IPluginManager.h b/Engine/Source/Runtime/Projects/Public/Interfaces/IPluginManager.h
--- a/Engine/Source/Runtime/Projects/Public/Interfaces/IPluginManager.h
+++ b/Engine/Source/Runtime/Projects/Public/Interfaces/IPluginManager.h
@@ -352,6 +352,12 @@ public:
 	virtual TSharedPtr<IPlugin> FindEnabledPluginFromPath(const FString& PluginPath) = 0;
 	virtual TSharedPtr<IPlugin> FindEnabledPluginFromDescriptor(const FPluginReferenceDescriptor& PluginDesc) = 0;
 
+	/**
+	 * Returns true when the given plugin can be enabled for the current running exe.
+	 */
+	virtual bool CanEnablePluginInCurrentTarget(const FStringView Name) = 0;
+	virtual bool CanEnablePluginInCurrentTarget(const ANSICHAR* Name) = 0;
+
 	/** 
 	 * Finds all plugin descriptors underneath a given directory (recursively)
 	 * @param Directory Search folder
