IPlugin::Get/SetVerseVersion

diff --git a/Engine/Source/Runtime/Projects/Private/PluginManager.h b/Engine/Source/Runtime/Projects/Private/PluginManager.h
--- a/Engine/Source/Runtime/Projects/Private/PluginManager.h
+++ b/Engine/Source/Runtime/Projects/Private/PluginManager.h
@@ -111,6 +111,16 @@ public:
 		Descriptor.VersePath = MoveTemp(InVersePath);
 	}
 
+	virtual TOptional<uint32> GetVerseVersion() const override
+	{
+		return Descriptor.VerseVersion;
+	}
+
+	virtual void SetVerseVersion(TOptional<uint32> InVerseVersion) override
+	{
+		Descriptor.VerseVersion = InVerseVersion;
+	}
+
 	virtual EPluginType GetType() const override
 	{
 		return Type;
diff --git a/Engine/Source/Runtime/Projects/Public/Interfaces/IPluginManager.h b/Engine/Source/Runtime/Projects/Public/Interfaces/IPluginManager.h
--- a/Engine/Source/Runtime/Projects/Public/Interfaces/IPluginManager.h
+++ b/Engine/Source/Runtime/Projects/Public/Interfaces/IPluginManager.h
@@ -179,6 +179,17 @@ public:
 	 */
 	virtual void SetVersePath(FString&& InVersePath) = 0;
 
+	/**
+	 * Return the Verse language version to compile the plugin with (if unspecified, the latest stable version is used)
+	 */
+	virtual TOptional<uint32> GetVerseVersion() const = 0;
+
+	/**
+	 * Sets the Verse language version to compile the plugin with (if unspecified, the latest stable version is used)
+	 * @param InVerseVersion Verse version to set on the plugin
+	 */
+	virtual void SetVerseVersion(TOptional<uint32> InVerseVersion) = 0;
+
 	/**
 	 * Returns the plugin's location
 	 *
