Adding CVar to control whether to IniCacheSet when loading plugins, since this shows an improvement on some platforms.

diff --git a/Engine/Source/Runtime/Projects/Private/PluginManager.cpp b/Engine/Source/Runtime/Projects/Private/PluginManager.cpp
--- a/Engine/Source/Runtime/Projects/Private/PluginManager.cpp
+++ b/Engine/Source/Runtime/Projects/Private/PluginManager.cpp
@@ -44,6 +44,10 @@ DEFINE_LOG_CATEGORY_STATIC( LogPluginManager, Log, All );
 	#define UE_DISABLE_PLUGIN_DISCOVERY 0
 #endif
 
+#if !defined(UE_PLUGIN_MANAGER_USE_INI_CACHE_SET)
+	#define UE_PLUGIN_MANAGER_USE_INI_CACHE_SET 0
+#endif
+
 namespace UE::PluginManager::Private
 {
 
@@ -1959,6 +1963,7 @@ bool FPluginManager::ConfigureEnabledPlugins()
 
 			// walk over each plugin, and add the config files that are named for the plugin
 			// this is a separate loop so that a plugin can modify another plugin's configs in the second loop below
+			const TSet<FString>* IniCacheSet = UE_PLUGIN_MANAGER_USE_INI_CACHE_SET ? &AllIniFiles : nullptr;
 			for (TSharedRef<FPlugin> PluginPtr : PluginsArray)
 			{
 				FPlugin& Plugin = *PluginPtr;
@@ -1971,6 +1976,7 @@ bool FPluginManager::ConfigureEnabledPlugins()
 				FConfigCacheIni::RegisterPlugin(PluginName, Plugin.GetBaseDir(), Plugin.GetExtensionBaseDirs(), DynamicLayerPriority::Plugin, bIncludePluginNameInBranchName);
 				
 				FConfigContext Context = FConfigContext::ReadIntoGConfig();
+				Context.IniCacheSet = IniCacheSet;
 				Context.ChangeTracker = &ChangeTracker;
 				Context.ConfigFileTag = *Plugin.Name;
 				Context.Load(*Plugin.Name);
