Plugin Manager: Prioritize editor target receipts whose Launch executable matches the currently running executable.

Previously, the first (arbitrary) editor target receipt with a matching configuration was chosen.

This is necessary to disambiguate a second receipt for an editor target in the engine binaries directory, and matches existing logic used to determine the correct receipt in FEngineLoop::AppInit.

diff --git a/Engine/Source/Developer/DesktopPlatform/Private/TargetReceipt.cpp b/Engine/Source/Developer/DesktopPlatform/Private/TargetReceipt.cpp
--- a/Engine/Source/Developer/DesktopPlatform/Private/TargetReceipt.cpp
+++ b/Engine/Source/Developer/DesktopPlatform/Private/TargetReceipt.cpp
@@ -1,6 +1,7 @@
 // Copyright Epic Games, Inc. All Rights Reserved.
 
 #include "TargetReceipt.h"
+#include "HAL/PlatformProcess.h"
 #include "Misc/Paths.h"
 #include "Misc/FileHelper.h"
 #include "Dom/JsonObject.h"
@@ -82,18 +83,33 @@ bool FTargetReceipt::Read(const FString& FileName, bool bExpandVariables)
 		return false;
 	}
 
-	// Get the launch path
-	if(!Object->TryGetStringField(TEXT("Launch"), Launch))
+	// Get the launch path(s). LaunchCmd is optional.
+	if (FString MaybeLaunchCmd; Object->TryGetStringField(TEXT("LaunchCmd"), MaybeLaunchCmd))
 	{
-		//If the regular launch field is not present try using the LaunchCmd field
-		if (!Object->TryGetStringField(TEXT("LaunchCmd"), Launch))
+		LaunchCmd = MoveTemp(MaybeLaunchCmd);
+	}
+
+	if (!Object->TryGetStringField(TEXT("Launch"), Launch))
+	{
+		// If the regular launch field is not present try using the LaunchCmd field
+		if (!LaunchCmd.IsSet())
 		{
 			return false;
 		}
+		else
+		{
+			Launch = LaunchCmd.GetValue();
+		}
 	}
+
 	if (bExpandVariables)
 	{
 		ExpandVariables(Launch);
+
+		if (LaunchCmd.IsSet())
+		{
+			ExpandVariables(LaunchCmd.GetValue());
+		}
 	}
 
 	// Read the list of build products
@@ -210,6 +226,23 @@ bool FTargetReceipt::Read(const FString& FileName, bool bExpandVariables)
 	return true;
 }
 
+bool FTargetReceipt::LaunchesCurrentExecutable(bool bCheckLaunchField, bool bCheckLaunchCmdField) const
+{
+	const FString CurrentExecutable = FPlatformProcess::ExecutablePath();
+
+	if (bCheckLaunchField && FPaths::IsSamePath(Launch, CurrentExecutable))
+	{
+		return true;
+	}
+
+	if (bCheckLaunchCmdField && LaunchCmd.IsSet() && FPaths::IsSamePath(LaunchCmd.GetValue(), CurrentExecutable))
+	{
+		return true;
+	}
+
+	return false;
+}
+
 FString FTargetReceipt::GetDefaultPath(const TCHAR* BaseDir, const TCHAR* TargetName, const TCHAR* Platform, EBuildConfiguration Configuration, const TCHAR* BuildArchitecture)
 {
 	const TCHAR* ArchitectureSuffix = TEXT("");
diff --git a/Engine/Source/Developer/DesktopPlatform/Public/TargetReceipt.h b/Engine/Source/Developer/DesktopPlatform/Public/TargetReceipt.h
--- a/Engine/Source/Developer/DesktopPlatform/Public/TargetReceipt.h
+++ b/Engine/Source/Developer/DesktopPlatform/Public/TargetReceipt.h
@@ -8,6 +8,7 @@
 #include "CoreMinimal.h"
 #include "HAL/Platform.h"
 #include "HAL/PlatformMisc.h"
+#include "Misc/Optional.h"
 #include "Modules/BuildVersion.h"
 
 /**
@@ -75,9 +76,15 @@ struct DESKTOPPLATFORM_API FTargetReceipt
 	/** The version information for this target */
 	FBuildVersion Version;
 
-	/** The exectuable to launch for this target */
+	/**
+	 * The executable to launch for this target.
+	 * (NOTE: If the LaunchCmd field exists, but Launch does not, this will actually reflect LaunchCmd!)
+	 */
 	FString Launch;
 
+	/** The console subsystem executable for this target */
+	TOptional<FString> LaunchCmd;
+
 	/** The build products which are part of this target */
 	TArray<FBuildProduct> BuildProducts;
 
@@ -102,6 +109,15 @@ struct DESKTOPPLATFORM_API FTargetReceipt
 	 */
 	bool Read(const FString& FileName, bool bExpandVariables = true);
 
+	/**
+	 * Compares the running executable's path to the `Launch` and/or `LaunchCmd` paths.
+	 * 
+	 * @param bCheckLaunchField Return true if the current executable is the one referred to by the `Launch` field.
+	 * @param bCheckLaunchCmdField Return true if the current executable is the one referred to by the `LaunchCmd` field.
+	 * @return True if either specified check matches; otherwise, false.
+	 */
+	bool LaunchesCurrentExecutable(bool bCheckLaunchField = true, bool bCheckLaunchCmdField = true) const;
+
 	/**
 	 * Gets the default path for a target receipt
 	 *
diff --git a/Engine/Source/Runtime/Launch/Private/LaunchEngineLoop.cpp b/Engine/Source/Runtime/Launch/Private/LaunchEngineLoop.cpp
--- a/Engine/Source/Runtime/Launch/Private/LaunchEngineLoop.cpp
+++ b/Engine/Source/Runtime/Launch/Private/LaunchEngineLoop.cpp
@@ -6291,22 +6291,15 @@ void CleanUpPath(FString& InPath)
 
 bool DoesCurrentExecutableMatch(const FString& EditorTargetFileName)
 {
-	// Figure out the executable that we should be running
-	FString LaunchExecutableName;
 	FTargetReceipt Receipt;
 	if (!FPaths::FileExists(EditorTargetFileName) || !Receipt.Read(EditorTargetFileName))
 	{
 		return false;
 	}
-	LaunchExecutableName = Receipt.Launch;
-	CleanUpPath(LaunchExecutableName);
 
-	// Get the current executable name. Don't allow relaunching if we're running the console app.
-	FString CurrentExecutableName = FPlatformProcess::ExecutablePath();
-	CleanUpPath(CurrentExecutableName);
-
-	// Nothing to do if they're the same
-	return FPaths::IsSamePath(LaunchExecutableName, CurrentExecutableName);
+	// Returning false for -Cmd matches legacy behavior; do we want to change this?
+	const bool bCheckLaunch_true = true, bCheckCmd_false = false;
+	return Receipt.LaunchesCurrentExecutable(bCheckLaunch_true, bCheckCmd_false);
 }
 
 bool LaunchCorrectEditorExecutable(const FString& EditorTargetFileName)
@@ -6516,6 +6509,7 @@ bool FEngineLoop::AppInit( )
 				{
 					ScratchEditorTargetFileName = FTargetReceipt::GetDefaultPath(*FPaths::EngineDir(), *Target.Name, FPlatformProcess::GetBinariesSubdirectory(), FApp::GetBuildConfiguration(), nullptr);
 				}
+
 				if (DoesCurrentExecutableMatch(ScratchEditorTargetFileName))
 				{
 					EditorTargetFileName = MoveTemp(ScratchEditorTargetFileName);
diff --git a/Engine/Source/Runtime/Projects/Private/PluginManager.cpp b/Engine/Source/Runtime/Projects/Private/PluginManager.cpp
--- a/Engine/Source/Runtime/Projects/Private/PluginManager.cpp
+++ b/Engine/Source/Runtime/Projects/Private/PluginManager.cpp
@@ -1419,29 +1419,55 @@ bool FPluginManager::ConfigureEnabledPlugins()
 		FString DefaultEditorTarget;
 		GConfig->GetString(TEXT("/Script/BuildSettings.BuildSettings"), TEXT("DefaultEditorTarget"), DefaultEditorTarget, GEngineIni);
 
-		auto FindFirstMatchingTargetFile = [&DefaultEditorTarget](const FString& ReceiptWildcard) -> TUniquePtr<FTargetReceipt>
+		auto FindMatchingTargetFile = [&DefaultEditorTarget](const FString& ReceiptWildcard) -> TUniquePtr<FTargetReceipt>
 		{
 			TArray<FString> AllTargetFilesWithoutPath;
 			const FString ReceiptPath = FPaths::GetPath(ReceiptWildcard);
 			IFileManager::Get().FindFiles(AllTargetFilesWithoutPath, *ReceiptWildcard, true, false);
 
+			// Used as a fallback in lieu of an exact match (only for EBuildTargetType::Editor).
+			TPair<FString, TUniquePtr<FTargetReceipt>> BestMatch;
+
 			for (const FString& TargetFileWithoutPath : AllTargetFilesWithoutPath)
 			{
-				const FString TargetFile = FPaths::Combine(ReceiptPath, TargetFileWithoutPath);
+				FString TargetFile = FPaths::Combine(ReceiptPath, TargetFileWithoutPath);
 				TUniquePtr<FTargetReceipt> Receipt = MakeUnique<FTargetReceipt>();
 				if (Receipt->Read(TargetFile))
 				{
 					if (Receipt->TargetType == FApp::GetBuildTargetType() && Receipt->Configuration == FApp::GetBuildConfiguration())
 					{
-						bool bIsDefaultTarget = Receipt->TargetType != EBuildTargetType::Editor || (DefaultEditorTarget.Len() == 0) || (DefaultEditorTarget == Receipt->TargetName);
-						if (bIsDefaultTarget)
+						// Priority:
+						// - If we are non-editor, first non-editor receipt
+						// - Otherwise, first that matches the running executable
+						// - Otherwise, first corresponding to DefaultEditorTarget (if specified)
+						// - Otherwise, first editor receipt
+						if (Receipt->TargetType != EBuildTargetType::Editor || Receipt->LaunchesCurrentExecutable())
 						{
+							UE_LOG(LogPluginManager, Log, TEXT("Found matching target receipt: %s"), *TargetFile);
 							return Receipt;
 						}
+
+						const bool bIsDefaultTarget = (DefaultEditorTarget.Len() != 0) && (DefaultEditorTarget == Receipt->TargetName);
+						if (bIsDefaultTarget || !BestMatch.Value)
+						{
+							BestMatch = { MoveTemp(TargetFile), MoveTemp(Receipt) };
+							// `TargetFile` and `Receipt` are now invalidated by virtue of being moved-from.
+						}
 					}
 				}
 			}
-			return TUniquePtr<FTargetReceipt>();
+
+			if (BestMatch.Value)
+			{
+				UE_LOG(LogPluginManager, Log, TEXT("Found target receipt: %s"), *BestMatch.Key);
+			}
+			else
+			{
+				UE_LOG(LogPluginManager, Log, TEXT("Unable to find target receipt in path: %s"), *ReceiptWildcard);
+			}
+
+			// Might be null.
+			return MoveTemp(BestMatch.Value);
 		};
 #endif // READ_TARGET_ENABLED_PLUGINS_FROM_RECEIPT
 
@@ -1457,11 +1483,11 @@ bool FPluginManager::ConfigureEnabledPlugins()
 			SCOPED_BOOT_TIMING("ReadTargetBuildPluginsFromReceipt");
 
 			// Read the build plugins from the target file using the target receipt file. This controls which optional plugin references can be enabled.
-			auto ReadBuildPluginsFromFirstMatchingTargetFile = [&FindFirstMatchingTargetFile, &AllowedOptionalDependencies](const TCHAR* BaseDir) -> bool
+			auto ReadBuildPluginsFromFirstMatchingTargetFile = [&FindMatchingTargetFile, &AllowedOptionalDependencies](const TCHAR* BaseDir) -> bool
 			{
 				const FString ReceiptWildcard = FTargetReceipt::GetDefaultPath(BaseDir, TEXT("*"), FPlatformProcess::GetBinariesSubdirectory(), FApp::GetBuildConfiguration(), nullptr);
 
-				TUniquePtr<FTargetReceipt> Receipt = FindFirstMatchingTargetFile(ReceiptWildcard);
+				TUniquePtr<FTargetReceipt> Receipt = FindMatchingTargetFile(ReceiptWildcard);
 				if (Receipt.IsValid())
 				{
 					AllowedOptionalDependencies.Append(Receipt->BuildPlugins);
@@ -1470,6 +1496,7 @@ bool FPluginManager::ConfigureEnabledPlugins()
 				return false;
 			};
 
+			UE_LOG(LogPluginManager, Log, TEXT("Looking for build plugins target receipt"));
 			if (!ReadBuildPluginsFromFirstMatchingTargetFile(FPlatformMisc::ProjectDir()))
 			{
 				ReadBuildPluginsFromFirstMatchingTargetFile(FPlatformMisc::EngineDir());
@@ -1607,14 +1634,14 @@ bool FPluginManager::ConfigureEnabledPlugins()
 
 #if READ_TARGET_ENABLED_PLUGINS_FROM_RECEIPT
 			// Configure the plugins that were enabled or disabled from the target file using the target receipt file
-			auto ConfigurePluginsFromFirstMatchingTargetFile = [this, &FindFirstMatchingTargetFile, &ConfiguredPluginNames,
+			auto ConfigurePluginsFromFirstMatchingTargetFile = [this, &FindMatchingTargetFile, &ConfiguredPluginNames,
 				&EnabledPlugins, &AllowedOptionalDependencies, &bAllowEnginePluginsEnabledByDefault]
 				(const TCHAR* BaseDir, bool& bOutError) -> bool
 			{
 				const FString ReceiptWildcard = FTargetReceipt::GetDefaultPath(BaseDir, TEXT("*"), FPlatformProcess::GetBinariesSubdirectory(), FApp::GetBuildConfiguration(), nullptr);
 				FString SourceDescription = FString::Printf(TEXT("Receipt files %s"), *ReceiptWildcard);
 
-				TUniquePtr<FTargetReceipt> Receipt = FindFirstMatchingTargetFile(ReceiptWildcard);
+				TUniquePtr<FTargetReceipt> Receipt = FindMatchingTargetFile(ReceiptWildcard);
 				if (Receipt.IsValid())
 				{
 					FReceiptProperty* AllowDefaultProperty = Algo::FindByPredicate(Receipt->AdditionalProperties,
@@ -1648,6 +1675,7 @@ bool FPluginManager::ConfigureEnabledPlugins()
 
 			{
 				SCOPED_BOOT_TIMING("ConfigureTargetEnabledPluginsFromReceipt");
+				UE_LOG(LogPluginManager, Log, TEXT("Looking for enabled plugins target receipt"));
 				bool bErrorConfiguring = false;
 				if (!ConfigurePluginsFromFirstMatchingTargetFile(FPlatformMisc::ProjectDir(), bErrorConfiguring))
 				{
