Removing IniCacheSet usage in plugin manager; the set doesn't actually get propagated in the ParallelFor below where it is currently used, and would require some thread-local globals or parameters to propagate it down safely to where it's needed. After testing whether that would be helpful, there are still some DirectoryExists calls later that would not be feasible to cache, so this just reverts it to the previous state where it wasn't being used at all.

diff --git a/Engine/Source/Runtime/Projects/Private/PluginManager.cpp b/Engine/Source/Runtime/Projects/Private/PluginManager.cpp
--- a/Engine/Source/Runtime/Projects/Private/PluginManager.cpp
+++ b/Engine/Source/Runtime/Projects/Private/PluginManager.cpp
@@ -44,10 +44,6 @@ DEFINE_LOG_CATEGORY_STATIC( LogPluginManager, Log, All );
 	#define UE_DISABLE_PLUGIN_DISCOVERY 0
 #endif
 
-#if !defined(UE_PLUGIN_MANAGER_USE_INI_CACHE_SET)
-	#define UE_PLUGIN_MANAGER_USE_INI_CACHE_SET 0
-#endif
-
 namespace UE::PluginManager::Private
 {
 
@@ -141,16 +137,6 @@ namespace PluginSystemDefs
 			PluginPathsOut.Add(FExternalPluginPath{ Path, EPluginExternalSource::Environment });
 		}
 	}
-
-
-	bool IsCachingIniFilesForProcessing()
-	{
-#if PLATFORM_DESKTOP // with the reduced set of plugin files to scan, this is likely unnecessary on any platform, but DESKTOP platforms may have Saved/Cooked directories around are slooow to scan
-		return false;
-#else
-		return true;
-#endif
-	}
 }
 
 /** 
@@ -1961,23 +1947,6 @@ bool FPluginManager::ConfigureEnabledPlugins()
 
 			TArray<FString> ConfigFilesPluginsCannotOverride;
 			GConfig->GetArray(TEXT("Plugins"), TEXT("ConfigFilesPluginsCannotOverride"), ConfigFilesPluginsCannotOverride, GEngineIni);
-
-			TSet<FString> AllIniFiles;
-
-			if (PluginSystemDefs::IsCachingIniFilesForProcessing())
-			{
-				SCOPED_BOOT_TIMING("ParallelPluginEnabling::FindIniFiles");
-				TArray<FString> AllIniFilesList;
-				
-				// Using ProjectDir and EngineDir are really broad, but they also cover all plugin config directories individually
-				// and should cover all of the extension directories as well.
-				IFileManager::Get().FindFilesRecursive(AllIniFilesList, *FPaths::EngineDir(), TEXT("*.ini"), /*Files=*/true, /*Directories=*/false, /*bClearFileNames=*/false);
-				IFileManager::Get().FindFilesRecursive(AllIniFilesList, *FPaths::ProjectDir(), TEXT("*.ini"), /*Files=*/true, /*Directories=*/false, /*bClearFileNames=*/false);
-				IFileManager::Get().FindFilesRecursive(AllIniFilesList, *FPaths::ProjectSavedDir(), TEXT("*.ini"), /*Files=*/true, /*Directories=*/false, /*bClearFileNames=*/false);
-
-				AllIniFiles.Append(AllIniFilesList);
-			}
-
 			
 			FCriticalSection ConfigCS;
 			FCriticalSection PluginPakCS;
@@ -1997,7 +1966,6 @@ bool FPluginManager::ConfigureEnabledPlugins()
 
 			// walk over each plugin, and add the config files that are named for the plugin
 			// this is a separate loop so that a plugin can modify another plugin's configs in the second loop below
-			const TSet<FString>* IniCacheSet = UE_PLUGIN_MANAGER_USE_INI_CACHE_SET ? &AllIniFiles : nullptr;
 			for (TSharedRef<FPlugin> PluginPtr : PluginsArray)
 			{
 				FPlugin& Plugin = *PluginPtr;
@@ -2010,14 +1978,13 @@ bool FPluginManager::ConfigureEnabledPlugins()
 				FConfigCacheIni::RegisterPlugin(PluginName, Plugin.GetBaseDir(), Plugin.GetExtensionBaseDirs(), DynamicLayerPriority::Plugin, bIncludePluginNameInBranchName);
 				
 				FConfigContext Context = FConfigContext::ReadIntoGConfig();
-				Context.IniCacheSet = IniCacheSet;
 				Context.ChangeTracker = &ChangeTracker;
 				Context.ConfigFileTag = *Plugin.Name;
 				Context.Load(*Plugin.Name);
 			}
 
 			// walk over each plugin, modify existing ini's, mount the content, etc
-			ParallelFor(PluginsArray.Num(), [&PluginsArray, &ConfigCS, &PluginPakCS, &PendingConfigsCS, &PendingConfigs, &ConfigFilesPluginsCannotOverride, &AllIniFiles, &ChangeTracker, this](int32 Index)
+			ParallelFor(PluginsArray.Num(), [&PluginsArray, &ConfigCS, &PluginPakCS, &PendingConfigsCS, &PendingConfigs, &ConfigFilesPluginsCannotOverride, &ChangeTracker, this](int32 Index)
 			{
 				FPlugin& Plugin = *PluginsArray[Index];
 				FName PluginName(*Plugin.Name);
