Some asset manifest generation changes.  Note that persistence backwards compat checking does not yet work with this, though asset digests are written out by the compiler when requested by the cooker.

diff --git a/Engine/Source/Runtime/Projects/Private/PluginManager.cpp b/Engine/Source/Runtime/Projects/Private/PluginManager.cpp
--- a/Engine/Source/Runtime/Projects/Private/PluginManager.cpp
+++ b/Engine/Source/Runtime/Projects/Private/PluginManager.cpp
@@ -3047,7 +3047,7 @@ TArray<TSharedRef<IPlugin>> FPluginManager::GetEnabledPluginsWithVerse() const
 	{
 		const TSharedRef<FPlugin>& PluginRef = DiscoveredPluginMapUtils::ResolvePluginFromMapVal(PluginPair.Value);
 		const FPlugin& Plugin = *PluginRef;
-		if (Plugin.IsEnabled() && Plugin.CanContainVerse() && Plugin.IsMounted())
+		if (Plugin.IsEnabled() && (Plugin.CanContainVerse() || (Plugin.CanContainContent() && Plugin.GetDescriptor().bEnableVerseAssetReflection)) && Plugin.IsMounted())
 		{
 			Plugins.Add(PluginRef);
 		}
