Explosing Plugin.IsMounted() as part of the Plugin interface

diff --git a/Engine/Source/Runtime/Projects/Private/PluginManager.h b/Engine/Source/Runtime/Projects/Private/PluginManager.h
--- a/Engine/Source/Runtime/Projects/Private/PluginManager.h
+++ b/Engine/Source/Runtime/Projects/Private/PluginManager.h
@@ -86,6 +86,11 @@ public:
 
 	virtual bool IsEnabledByDefault(bool bAllowEnginePluginsEnabledByDefault) const override;
 
+	virtual bool IsMounted() const override
+	{
+		return bIsMounted;
+	}
+
 	virtual bool IsHidden() const override
 	{
 		return Descriptor.bIsHidden;
@@ -143,7 +148,6 @@ public:
 	virtual const TSharedPtr<FJsonObject>& GetDescriptorJson() override;
 #endif // WITH_EDITOR
 
-	bool IsMounted() const { return bIsMounted; }
 	void SetIsMounted(bool bInIsMounted) { bIsMounted = bInIsMounted; }
 };
 
diff --git a/Engine/Source/Runtime/Projects/Public/Interfaces/IPluginManager.h b/Engine/Source/Runtime/Projects/Public/Interfaces/IPluginManager.h
--- a/Engine/Source/Runtime/Projects/Public/Interfaces/IPluginManager.h
+++ b/Engine/Source/Runtime/Projects/Public/Interfaces/IPluginManager.h
@@ -181,6 +181,13 @@ public:
 	 */
 	virtual bool IsEnabledByDefault(bool bAllowEnginePluginsEnabledByDefault) const = 0;
 
+	/**
+	 * Determines if the plugin is mounted.
+	 *
+	 * @return True if the plugin is currently mounted.
+	 */
+	virtual bool IsMounted() const = 0;
+
 	/**
 	 * Determines if the plugin is should be displayed in-editor for the user to enable/disable freely.
 	 *
