Silence an instance of static analysis warning related to iteration on a container that might be empty.

diff --git a/Engine/Source/Runtime/Projects/Private/PluginManager.cpp b/Engine/Source/Runtime/Projects/Private/PluginManager.cpp
--- a/Engine/Source/Runtime/Projects/Private/PluginManager.cpp
+++ b/Engine/Source/Runtime/Projects/Private/PluginManager.cpp
@@ -1455,7 +1455,7 @@ bool FPluginManager::ConfigureEnabledPlugins()
 					return false;
 				});
 
-				for (const FString& PluginListFilename : PluginListFilenames)
+				for (const FString& PluginListFilename : PluginListFilenames) //-V1078
 				{
 					FFileHelper::LoadFileToStringWithLineVisitor(*PluginListFilename, [&PluginsList](FStringView Line)
 					{
