Log plugin name in FindPluginInMap_Checked

diff --git a/Engine/Source/Runtime/Projects/Private/PluginManager.cpp b/Engine/Source/Runtime/Projects/Private/PluginManager.cpp
--- a/Engine/Source/Runtime/Projects/Private/PluginManager.cpp
+++ b/Engine/Source/Runtime/Projects/Private/PluginManager.cpp
@@ -212,7 +212,7 @@ namespace DiscoveredPluginMapUtils
 	static TSharedRef<FPlugin>& FindPluginInMap_Checked(FDiscoveredPluginMap& InMap, const FDiscoveredPluginMap::KeyType& Name)
 	{
 		TSharedRef<FPlugin>* PluginPtr = FindPluginInMap(InMap, Name);
-		check(PluginPtr != nullptr);
+		checkf(PluginPtr, TEXT("Cannot find plugin %s"), *Name);
 		return *PluginPtr;
 	}
 
