Removed IPlugin::SetVersePath and similar setters.  Users should call UpdateDescriptor so we get an event when the descriptor changes.

diff --git a/Engine/Source/Runtime/Projects/Private/PluginManager.h b/Engine/Source/Runtime/Projects/Private/PluginManager.h
--- a/Engine/Source/Runtime/Projects/Private/PluginManager.h
+++ b/Engine/Source/Runtime/Projects/Private/PluginManager.h
@@ -111,31 +111,16 @@ public:
 		return Descriptor.VersePath;
 	}
 
-	virtual void SetVersePath(FString&& InVersePath) override
-	{
-		Descriptor.VersePath = MoveTemp(InVersePath);
-	}
-
 	virtual TOptional<uint32> GetVerseVersion() const override
 	{
 		return Descriptor.VerseVersion;
 	}
 
-	virtual void SetVerseVersion(TOptional<uint32> InVerseVersion) override
-	{
-		Descriptor.VerseVersion = InVerseVersion;
-	}
-
 	virtual bool IsSceneGraphEnabled() const override
 	{
 		return Descriptor.bEnableSceneGraph;
 	}
 
-	virtual void SetEnableSceneGraph(bool bInEnableSceneGraph) override
-	{
-		Descriptor.bEnableSceneGraph = bInEnableSceneGraph;
-	}
-
 	virtual EPluginType GetType() const override
 	{
 		return Type;
diff --git a/Engine/Source/Runtime/Projects/Public/Interfaces/IPluginManager.h b/Engine/Source/Runtime/Projects/Public/Interfaces/IPluginManager.h
--- a/Engine/Source/Runtime/Projects/Public/Interfaces/IPluginManager.h
+++ b/Engine/Source/Runtime/Projects/Public/Interfaces/IPluginManager.h
@@ -216,34 +216,16 @@ public:
 	 */
 	virtual const FString& GetVersePath() const = 0;
 
-	/**
-	 * Sets the Verse path to the root of the plugin's content directory
-	 * @param InVersePath Verse path to set
-	 */
-	virtual void SetVersePath(FString&& InVersePath) = 0;
-
 	/**
 	 * Return the Verse language version to compile the plugin with (if unspecified, the latest stable version is used)
 	 */
 	virtual TOptional<uint32> GetVerseVersion() const = 0;
 
-	/**
-	 * Sets the Verse language version to compile the plugin with (if unspecified, the latest stable version is used)
-	 * @param InVerseVersion Verse version to set on the plugin
-	 */
-	virtual void SetVerseVersion(TOptional<uint32> InVerseVersion) = 0;
-
 	/**
 	 * Return whether Scene Graph is enabled or not. This impacts which Verse asset digest is generated.
 	 */
 	virtual bool IsSceneGraphEnabled() const = 0;
 
-	/**
-	 * Sets whether Scene Graph is enabled or not. This impacts which Verse asset digest is generated.
-	 * @param bInEnableSceneGraph Whether or not Scene Graph is enabled on the plugin
-	 */
-	virtual void SetEnableSceneGraph(bool bInEnableSceneGraph) = 0;
-
 	/**
 	 * Returns the plugin's location
 	 *
