Don't load optional plugins that were not compiled in modular builds

diff --git a/Engine/Source/Runtime/Projects/Private/PluginManager.cpp b/Engine/Source/Runtime/Projects/Private/PluginManager.cpp
--- a/Engine/Source/Runtime/Projects/Private/PluginManager.cpp
+++ b/Engine/Source/Runtime/Projects/Private/PluginManager.cpp
@@ -2487,6 +2487,14 @@ bool FPluginManager::ConfigureEnabledPluginForTarget(const FPluginReferenceDescr
 		}
 #endif // IS_MONOLITHIC
 
+#if !UE_BUILD_SHIPPING
+		if (Reference.bOptional && !AllowedOptionalDependencies.IsEmpty() && !AllowedOptionalDependencies.Contains(Reference.Name))
+		{
+			UE_LOG(LogPluginManager, Display, TEXT("Ignored optional reference to '%s' plugin; plugin was not built by target."), *Reference.Name);
+			continue;
+		}
+#endif
+
 		// Add references to all its dependencies
 		for (const FPluginReferenceDescriptor& NextReference : Plugin.Descriptor.Plugins)
 		{
