Added support for being notified when an async localization data load has finished, and updated GFP loading to use that rather than blocking for each GFP loaded

FTextLocalizationManager::NotifyWhenAsyncTasksCompleted takes a function to call when the current set of async localization data loads has finished, or immediately if there are no pending async localization data loads. The notification function may be called on any thread.

Other functions on FTextLocalizationManager that start an async localization data load may now directly take an optional notification callback, which is called when that localization data load has finished. The notification function may be called on any thread.

GFP loading now uses a notification callback rather than use FTextLocalizationManager::WaitForAsyncTasks to determine when the async localization data loads it has triggered have finished, and FGameFeaturePluginState_Activating will now watch for that notification result rather than block. This allows other GFPs to continue loading while the async localization data load finishes.

diff --git a/Engine/Plugins/Runtime/GameFeatures/Source/GameFeatures/Private/GameFeaturePluginStateMachine.cpp b/Engine/Plugins/Runtime/GameFeatures/Source/GameFeatures/Private/GameFeaturePluginStateMachine.cpp
--- a/Engine/Plugins/Runtime/GameFeatures/Source/GameFeatures/Private/GameFeaturePluginStateMachine.cpp
+++ b/Engine/Plugins/Runtime/GameFeatures/Source/GameFeatures/Private/GameFeaturePluginStateMachine.cpp
@@ -118,6 +118,37 @@ namespace UE::GameFeatures
 		return !GIsEditor && bDeferLocalizationDataLoad;
 	}
 
+	void MountLocalizationData(UGameFeaturePluginStateMachine* CurrentMachine, FGameFeaturePluginStateMachineProperties& StateProperties)
+	{
+		check(IsInGameThread());
+		check(CurrentMachine && &CurrentMachine->GetProperties() == &StateProperties);
+
+		StateProperties.bIsLoadingLocalizationData = true;
+		IPluginManager::Get().MountExplicitlyLoadedPluginLocalizationData(StateProperties.PluginName, [WeakMachine = TWeakObjectPtr<UGameFeaturePluginStateMachine>(CurrentMachine), &StateProperties](bool bLoadedLocalization, const FString& PluginName)
+		{
+			if (IsInGameThread())
+			{
+				if (UGameFeaturePluginStateMachine* StateMachine = WeakMachine.Get();
+					StateMachine && ensureAlways(&StateMachine->GetProperties() == &StateProperties))
+				{
+					StateProperties.bIsLoadingLocalizationData = false;
+				}
+			}
+			else
+			{
+				ExecuteOnGameThread(UE_SOURCE_LOCATION, [WeakMachine, &StateProperties]()
+				{
+					if (UGameFeaturePluginStateMachine* StateMachine = WeakMachine.Get();
+						StateMachine && ensureAlways(&StateMachine->GetProperties() == &StateProperties))
+					{
+						StateProperties.bIsLoadingLocalizationData = false;
+						StateProperties.OnRequestUpdateStateMachine.ExecuteIfBound();
+					}
+				});
+			}
+		});
+	}
+
 	bool ShouldSkipVerify(const FString& PluginName)
 	{
 		static const FAsciiSet Wildcards("*?");
@@ -2488,7 +2519,8 @@ struct FGameFeaturePluginState_Mounting : public FGameFeaturePluginState
 			verify(IPluginManager::Get().MountExplicitlyLoadedPlugin(StateProperties.PluginName));
 			if (!UE::GameFeatures::ShouldDeferLocalizationDataLoad())
 			{
-				IPluginManager::Get().MountExplicitlyLoadedPluginLocalizationData(StateProperties.PluginName);
+				UGameFeaturePluginStateMachine* CurrentMachine = UGameFeaturesSubsystem::Get().FindGameFeaturePluginStateMachine(StateProperties.PluginIdentifier);
+				UE::GameFeatures::MountLocalizationData(CurrentMachine, StateProperties);
 			}
 			if (bManuallyOpenPluginShaderLibrary)
 			{
@@ -2502,7 +2534,8 @@ struct FGameFeaturePluginState_Mounting : public FGameFeaturePluginState
 		verify(IPluginManager::Get().MountExplicitlyLoadedPlugin(StateProperties.PluginName));
 		if (!UE::GameFeatures::ShouldDeferLocalizationDataLoad())
 		{
-			IPluginManager::Get().MountExplicitlyLoadedPluginLocalizationData(StateProperties.PluginName);
+			UGameFeaturePluginStateMachine* CurrentMachine = UGameFeaturesSubsystem::Get().FindGameFeaturePluginStateMachine(StateProperties.PluginIdentifier);
+			UE::GameFeatures::MountLocalizationData(CurrentMachine, StateProperties);
 		}
 
 		// Now load the shader lib in the background
@@ -3347,7 +3380,10 @@ struct FGameFeaturePluginState_Unloading : public FGameFeaturePluginState
 
 	virtual void BeginState() override
 	{
-		IPluginManager::Get().UnmountExplicitlyLoadedPluginLocalizationData(StateProperties.PluginName);
+		if (UE::GameFeatures::ShouldDeferLocalizationDataLoad())
+		{
+			IPluginManager::Get().UnmountExplicitlyLoadedPluginLocalizationData(StateProperties.PluginName);
+		}
 	}
 
 	virtual void UpdateState(FGameFeaturePluginStateStatus& StateStatus) override
@@ -3396,7 +3432,11 @@ struct FGameFeaturePluginState_Loading : public FGameFeaturePluginState
 		TRACE_CPUPROFILER_EVENT_SCOPE(GFP_Loading_Begin);
 		check(StateProperties.GameFeatureData);
 
-		IPluginManager::Get().MountExplicitlyLoadedPluginLocalizationData(StateProperties.PluginName);
+		if (UE::GameFeatures::ShouldDeferLocalizationDataLoad())
+		{
+			UGameFeaturePluginStateMachine* CurrentMachine = UGameFeaturesSubsystem::Get().FindGameFeaturePluginStateMachine(StateProperties.PluginIdentifier);
+			UE::GameFeatures::MountLocalizationData(CurrentMachine, StateProperties);
+		}
 
 		BundleHandle = LoadGameFeatureBundles(StateProperties.GameFeatureData);
 	}
@@ -3729,6 +3769,12 @@ struct FGameFeaturePluginState_Activating : public FGameFeaturePluginState
 		check(GEngine);
 		check(StateProperties.GameFeatureData);
 
+		// If this plugin caused localization data to load, we need that load to finish before marking it as active
+		if (StateProperties.bIsLoadingLocalizationData)
+		{
+			return;
+		}
+
 		if (IsWaitingForBatchProcessing())
 		{
 			return;
@@ -3749,15 +3795,6 @@ struct FGameFeaturePluginState_Activating : public FGameFeaturePluginState
 			UGameFeaturesSubsystem::Get().OnGameFeatureActivating(StateProperties.GameFeatureData, StateProperties.PluginName, Context, StateProperties.PluginIdentifier);
 		}
 
-		// @TODO: non-blocking wait here?
-		// If this plugin caused localization data to load, wait for that here before marking it as active
-		if (TSharedPtr<IPlugin> Plugin = IPluginManager::Get().FindPlugin(StateProperties.PluginName);
-			Plugin && Plugin->GetDescriptor().bExplicitlyLoaded && Plugin->GetDescriptor().LocalizationTargets.Num() > 0)
-		{
-			TRACE_CPUPROFILER_EVENT_SCOPE(GFP_Activating_WaitForLoc);
-			FTextLocalizationManager::Get().WaitForAsyncTasks();
-		}
-
 		StateStatus.SetTransition(EGameFeaturePluginState::Active);
 	}
 
diff --git a/Engine/Plugins/Runtime/GameFeatures/Source/GameFeatures/Private/GameFeaturePluginStateMachine.h b/Engine/Plugins/Runtime/GameFeatures/Source/GameFeatures/Private/GameFeaturePluginStateMachine.h
--- a/Engine/Plugins/Runtime/GameFeatures/Source/GameFeatures/Private/GameFeaturePluginStateMachine.h
+++ b/Engine/Plugins/Runtime/GameFeatures/Source/GameFeatures/Private/GameFeaturePluginStateMachine.h
@@ -291,6 +291,9 @@ struct FGameFeaturePluginStateMachineProperties
 	/** The desired state during a transition. */
 	FGameFeaturePluginStateRange Destination;
 
+	/** Whether this plugin has an async localization load in-flight */
+	bool bIsLoadingLocalizationData : 1 = false;
+
 	/** Tracks whether or not this state machine added the plugin to the plugin manager. */
 	bool bAddedPluginToManager : 1 = false;
 
@@ -300,13 +303,13 @@ struct FGameFeaturePluginStateMachineProperties
 	/** Whether this state machine should attempt to cancel the current transition */
 	bool bTryCancel : 1 = false;
 
-	TArray<FName> AddedPrimaryAssetTypes;
+	/** Tracks if the current state was batch processed */
+	bool bWasBatchProcessed : 1 = false;
 
 	/** Batch Processing handle if requested for batch processed */
 	FDelegateHandle BatchProcessingHandle;
 
-	/** Tracks if the current state was batch processed */
-	bool bWasBatchProcessed : 1 = false;
+	TArray<FName> AddedPrimaryAssetTypes;
 
 	/** The data asset describing this game feature */
 	UPROPERTY(Transient)
diff --git a/Engine/Source/Runtime/Core/Private/Internationalization/TextLocalizationManager.cpp b/Engine/Source/Runtime/Core/Private/Internationalization/TextLocalizationManager.cpp
--- a/Engine/Source/Runtime/Core/Private/Internationalization/TextLocalizationManager.cpp
+++ b/Engine/Source/Runtime/Core/Private/Internationalization/TextLocalizationManager.cpp
@@ -1201,7 +1201,19 @@ void FTextLocalizationManager::WaitForAsyncTasks()
 	}
 }
 
-void FTextLocalizationManager::RefreshResources()
+void FTextLocalizationManager::NotifyWhenAsyncTasksCompleted(TUniqueFunction<void()>&& Notification)
+{
+	if (AsyncLocalizationTask)
+	{
+		QueueAsyncTask(MoveTemp(Notification));
+	}
+	else
+	{
+		Notification();
+	}
+}
+
+void FTextLocalizationManager::RefreshResources(TUniqueFunction<void()>&& Notification)
 {
 	ELocalizationLoadFlags LocLoadFlags = ELocalizationLoadFlags::None;
 	LocLoadFlags |= (WITH_EDITOR ? ELocalizationLoadFlags::Editor : ELocalizationLoadFlags::None);
@@ -1210,20 +1222,28 @@ void FTextLocalizationManager::RefreshResources()
 	LocLoadFlags |= ELocalizationLoadFlags::Native;
 	LocLoadFlags |= ELocalizationLoadFlags::Additional;
 
-	LoadLocalizationResourcesForCulture_Async(FInternationalization::Get().GetCurrentLanguage()->GetName(), LocLoadFlags);
+	LoadLocalizationResourcesForCulture_Async(FInternationalization::Get().GetCurrentLanguage()->GetName(), LocLoadFlags, MoveTemp(Notification));
 }
 
-void FTextLocalizationManager::HandleLocalizationTargetsMounted(TArrayView<const FString> LocalizationTargetPaths)
+void FTextLocalizationManager::HandleLocalizationTargetsMounted(TArrayView<const FString> LocalizationTargetPaths, TUniqueFunction<void()>&& Notification)
 {
 	if (!IsInitialized() || LocalizationTargetPaths.IsEmpty())
 	{
 		// If we've not yet loaded localization data then there's nothing to do
+		if (Notification)
+		{
+			Notification();
+		}
 		return;
 	}
 
 	// Nothing to do?
 	if (!FTextLocalizationManager::IsDisplayStringSupportEnabled())
 	{
+		if (Notification)
+		{
+			Notification();
+		}
 		return;
 	}
 
@@ -1247,25 +1267,33 @@ void FTextLocalizationManager::HandleLocalizationTargetsMounted(TArrayView<const
 
 	const TArray<FString> PrioritizedCultureNames = FInternationalization::Get().GetPrioritizedCultureNames(FInternationalization::Get().GetCurrentLanguage()->GetName());
 
-	LoadLocalizationTargetsForPrioritizedCultures_Async(LocalizationTargetPaths, PrioritizedCultureNames, LocLoadFlags);
+	LoadLocalizationTargetsForPrioritizedCultures_Async(LocalizationTargetPaths, PrioritizedCultureNames, LocLoadFlags, MoveTemp(Notification));
 }
 
-void FTextLocalizationManager::HandleLocalizationTargetsUnmounted(TArrayView<const FString> LocalizationTargetPaths)
+void FTextLocalizationManager::HandleLocalizationTargetsUnmounted(TArrayView<const FString> LocalizationTargetPaths, TUniqueFunction<void()>&& Notification)
 {
 	if (!IsInitialized() || LocalizationTargetPaths.IsEmpty())
 	{
 		// If we've not yet loaded localization data then there's nothing to do
+		if (Notification)
+		{
+			Notification();
+		}
 		return;
 	}
 	
 	// Nothing to do?
 	if (!FTextLocalizationManager::IsDisplayStringSupportEnabled())
 	{
+		if (Notification)
+		{
+			Notification();
+		}
 		return;
 	}
 
 	// Async update the live table
-	QueueAsyncTask([LocalizationTargetPaths = TArray<FString>(LocalizationTargetPaths)]()
+	QueueAsyncTask([LocalizationTargetPaths = TArray<FString>(LocalizationTargetPaths), Notification = MoveTemp(Notification)]()
 	{
 		TRACE_CPUPROFILER_EVENT_SCOPE(FTextLocalizationManager::HandleLocalizationTargetsUnmounted);
 		LLM_SCOPE_BYNAME(TEXT("Localization/DisplayStrings"));
@@ -1295,6 +1323,11 @@ void FTextLocalizationManager::HandleLocalizationTargetsUnmounted(TArrayView<con
 
 		// Allow any lingering texts that were referencing the unloaded display strings to release their references
 		TLM.DirtyTextRevision();
+
+		if (Notification)
+		{
+			Notification();
+		}
 	});
 }
 
@@ -1366,11 +1399,15 @@ void FTextLocalizationManager::LoadLocalizationResourcesForCulture_Sync(TArrayVi
 	LoadLocalizationResourcesForPrioritizedCultures_Sync(AvailableTextSources, FInternationalization::Get().GetPrioritizedCultureNames(CultureName), LocLoadFlags);
 }
 
-void FTextLocalizationManager::LoadLocalizationResourcesForCulture_Async(const FString& CultureName, const ELocalizationLoadFlags LocLoadFlags)
+void FTextLocalizationManager::LoadLocalizationResourcesForCulture_Async(const FString& CultureName, const ELocalizationLoadFlags LocLoadFlags, TUniqueFunction<void()>&& Notification)
 {
-	QueueAsyncTask([AvailableTextSources = LocalizedTextSources, CultureName, LocLoadFlags]()
+	QueueAsyncTask([AvailableTextSources = LocalizedTextSources, CultureName, LocLoadFlags, Notification = MoveTemp(Notification)]()
 	{
 		FTextLocalizationManager::Get().LoadLocalizationResourcesForCulture_Sync(AvailableTextSources, CultureName, LocLoadFlags);
+		if (Notification)
+		{
+			Notification();
+		}
 	});
 }
 
@@ -1482,11 +1519,15 @@ void FTextLocalizationManager::LoadLocalizationTargetsForPrioritizedCultures_Syn
 	}
 }
 
-void FTextLocalizationManager::LoadLocalizationTargetsForPrioritizedCultures_Async(TArrayView<const FString> LocalizationTargetPaths, TArrayView<const FString> PrioritizedCultureNames, const ELocalizationLoadFlags LocLoadFlags)
+void FTextLocalizationManager::LoadLocalizationTargetsForPrioritizedCultures_Async(TArrayView<const FString> LocalizationTargetPaths, TArrayView<const FString> PrioritizedCultureNames, const ELocalizationLoadFlags LocLoadFlags, TUniqueFunction<void()>&& Notification)
 {
-	QueueAsyncTask([AvailableTextSources = LocalizedTextSources, LocalizationTargetPaths = TArray<FString>(LocalizationTargetPaths), PrioritizedCultureNames = TArray<FString>(PrioritizedCultureNames), LocLoadFlags]()
+	QueueAsyncTask([AvailableTextSources = LocalizedTextSources, LocalizationTargetPaths = TArray<FString>(LocalizationTargetPaths), PrioritizedCultureNames = TArray<FString>(PrioritizedCultureNames), LocLoadFlags, Notification = MoveTemp(Notification)]()
 	{
-		FTextLocalizationManager::Get().LoadLocalizationTargetsForPrioritizedCultures_Sync(AvailableTextSources, LocalizationTargetPaths, PrioritizedCultureNames, LocLoadFlags);;
+		FTextLocalizationManager::Get().LoadLocalizationTargetsForPrioritizedCultures_Sync(AvailableTextSources, LocalizationTargetPaths, PrioritizedCultureNames, LocLoadFlags);
+		if (Notification)
+		{
+			Notification();
+		}
 	});
 }
 
diff --git a/Engine/Source/Runtime/Core/Public/Internationalization/TextLocalizationManager.h b/Engine/Source/Runtime/Core/Public/Internationalization/TextLocalizationManager.h
--- a/Engine/Source/Runtime/Core/Public/Internationalization/TextLocalizationManager.h
+++ b/Engine/Source/Runtime/Core/Public/Internationalization/TextLocalizationManager.h
@@ -221,6 +221,7 @@ public:
 
 	/**
 	 * Wait for any current async tasks to finish.
+	 * @see NotifyWhenAsyncTasksCompleted for an non-blocking version.
 	 * 
 	 * Async tasks will be started for anything that loads localization data (eg, initialization, language changes, requests to refresh the localization data, requests 
 	 * to load additional localization data for chunked targets, or requests to load additional localization data for explicitly loaded plugins).
@@ -231,25 +232,32 @@ public:
 	 */
 	CORE_API void WaitForAsyncTasks();
 
+	/**
+	 * Call the given function when any current async tasks have finished.
+	 * @note The notification may be called from any thread, to avoid blocking other async tasks in the queue while waiting for a given thread to be available.
+	 * @see WaitForAsyncTasks for a blocking version.
+	 */
+	CORE_API void NotifyWhenAsyncTasksCompleted(TUniqueFunction<void()>&& Notification);
+
 	/**
 	 * Reloads resources for the current culture.
-	 * @note The loading is async, see WaitUntilLoadingCompletes.
+	 * @note The loading is async, see WaitForAsyncTasks/NotifyWhenAsyncTasksCompleted, or pass the optional notification callback.
 	 */
-	CORE_API void RefreshResources();
+	CORE_API void RefreshResources(TUniqueFunction<void()>&& Notification = nullptr);
 
 	/**
 	 * Called when paths containing additional localization target data (LocRes) are mounted, to allow the display strings to dynamically update without waiting for a refresh.
-	 * @note The loading is async, see WaitUntilLoadingCompletes.
+	 * @note The loading is async, see WaitForAsyncTasks/NotifyWhenAsyncTasksCompleted, or pass the optional notification callback.
 	 * @see FCoreDelegates::GatherAdditionalLocResPathsCallback.
 	 */
-	CORE_API void HandleLocalizationTargetsMounted(TArrayView<const FString> LocalizationTargetPaths);
+	CORE_API void HandleLocalizationTargetsMounted(TArrayView<const FString> LocalizationTargetPaths, TUniqueFunction<void()>&& Notification = nullptr);
 
 	 /**
 	  * Called when paths containing additional localization target data (LocRes) are unmounted, to allow the display strings to be unloaded.
-	  * @note The unloading is async, see WaitUntilLoadingCompletes.
+	  * @note The unloading is async, see WaitForAsyncTasks/NotifyWhenAsyncTasksCompleted, or pass the optional notification callback.
 	  * @see FCoreDelegates::GatherAdditionalLocResPathsCallback.
 	  */
-	CORE_API void HandleLocalizationTargetsUnmounted(TArrayView<const FString> LocalizationTargetPaths);
+	CORE_API void HandleLocalizationTargetsUnmounted(TArrayView<const FString> LocalizationTargetPaths, TUniqueFunction<void()>&& Notification = nullptr);
 
 	/**
 	 * Returns the current text revision number. This value can be cached when caching information from the text localization manager.
@@ -347,7 +355,7 @@ private:
 
 	/** Loads localization resources for the specified culture, optionally loading localization resources that are editor-specific or game-specific. */
 	void LoadLocalizationResourcesForCulture_Sync(TArrayView<const TSharedPtr<ILocalizedTextSource>> AvailableTextSources, const FString& CultureName, const ELocalizationLoadFlags LocLoadFlags);
-	void LoadLocalizationResourcesForCulture_Async(const FString& CultureName, const ELocalizationLoadFlags LocLoadFlags);
+	void LoadLocalizationResourcesForCulture_Async(const FString& CultureName, const ELocalizationLoadFlags LocLoadFlags, TUniqueFunction<void()>&& Notification);
 
 	/** Loads localization resources for the specified prioritized cultures, optionally loading localization resources that are editor-specific or game-specific. */
 	void LoadLocalizationResourcesForPrioritizedCultures_Sync(TArrayView<const TSharedPtr<ILocalizedTextSource>> AvailableTextSources, TArrayView<const FString> PrioritizedCultureNames, const ELocalizationLoadFlags LocLoadFlags);
@@ -355,7 +363,7 @@ private:
 
 	/** Loads the specified localization targets for the specified prioritized cultures, optionally loading localization resources that are editor-specific or game-specific. */
 	void LoadLocalizationTargetsForPrioritizedCultures_Sync(TArrayView<const TSharedPtr<ILocalizedTextSource>> AvailableTextSources, TArrayView<const FString> LocalizationTargetPaths, TArrayView<const FString> PrioritizedCultureNames, const ELocalizationLoadFlags LocLoadFlags);
-	void LoadLocalizationTargetsForPrioritizedCultures_Async(TArrayView<const FString> LocalizationTargetPaths, TArrayView<const FString> PrioritizedCultureNames, const ELocalizationLoadFlags LocLoadFlags);
+	void LoadLocalizationTargetsForPrioritizedCultures_Async(TArrayView<const FString> LocalizationTargetPaths, TArrayView<const FString> PrioritizedCultureNames, const ELocalizationLoadFlags LocLoadFlags, TUniqueFunction<void()>&& Notification);
 
 	/** Loads the chunked localization resource data for the specified chunk */
 	void LoadChunkedLocalizationResources_Sync(TArrayView<const TSharedPtr<ILocalizedTextSource>> AvailableTextSources, const int32 ChunkId, const FString& PakFilename);
diff --git a/Engine/Source/Runtime/Projects/Private/PluginManager.cpp b/Engine/Source/Runtime/Projects/Private/PluginManager.cpp
--- a/Engine/Source/Runtime/Projects/Private/PluginManager.cpp
+++ b/Engine/Source/Runtime/Projects/Private/PluginManager.cpp
@@ -3261,7 +3261,7 @@ bool FPluginManager::MountExplicitlyLoadedPlugin_FromDescriptor(const FPluginRef
 	return TryMountExplicitlyLoadedPluginVersion(DescribedPlugin);
 }
 
-bool FPluginManager::MountExplicitlyLoadedPluginLocalizationData(const FString& PluginName)
+bool FPluginManager::MountExplicitlyLoadedPluginLocalizationData(const FString& PluginName, TUniqueFunction<void(bool bLoadedLocalization, const FString& PluginName)>&& Notification)
 {
 	TSharedPtr<FPlugin> Plugin = FindPluginInstance(PluginName);
 
@@ -3270,6 +3270,10 @@ bool FPluginManager::MountExplicitlyLoadedPluginLocalizationData(const FString&
 	{
 		// Does not exist or is not mounted
 		UE_LOG(LogPluginManager, Error, TEXT("Cannot mount plugin localization for '%s' as the plugin is unknown or not mounted. Did you forget to call MountExplicitlyLoadedPlugin?"), *PluginName);
+		if (Notification)
+		{
+			Notification(false, PluginName);
+		}
 		return false;
 	}
 
@@ -3277,6 +3281,10 @@ bool FPluginManager::MountExplicitlyLoadedPluginLocalizationData(const FString&
 	{
 		// Not supported
 		UE_LOG(LogPluginManager, Warning, TEXT("Cannot mount plugin localization for '%s' as the plugin isn't explicitly loaded."), *PluginName);
+		if (Notification)
+		{
+			Notification(false, PluginName);
+		}
 		return false;
 	}
 
@@ -3285,6 +3293,10 @@ bool FPluginManager::MountExplicitlyLoadedPluginLocalizationData(const FString&
 		// Already loaded
 		UE_LOG(LogPluginManager, Verbose, TEXT("Ignoring request to mount plugin localization for '%s' as it was already referenced (ref-count: %d -> %d)."), *PluginName, Plugin->ExplicitlyLoadedLocalizationDataMountedRefCount, Plugin->ExplicitlyLoadedLocalizationDataMountedRefCount + 1);
 		++Plugin->ExplicitlyLoadedLocalizationDataMountedRefCount;
+		if (Notification)
+		{
+			Notification(false, PluginName);
+		}
 		return false;
 	}
 
@@ -3292,6 +3304,10 @@ bool FPluginManager::MountExplicitlyLoadedPluginLocalizationData(const FString&
 	{
 		// Nothing to load
 		UE_LOG(LogPluginManager, Verbose, TEXT("Ignoring request to mount plugin localization for '%s' as the plugin has no localization targets defined."), *PluginName);
+		if (Notification)
+		{
+			Notification(false, PluginName);
+		}
 		return false;
 	}
 
@@ -3301,11 +3317,21 @@ bool FPluginManager::MountExplicitlyLoadedPluginLocalizationData(const FString&
 	// Notify that additional localization data should be loaded
 	TArray<FString> AdditionalLocResPaths;
 	PluginLocalizationUtils::GetLocalizationPathsForPlugin(*Plugin, AdditionalLocResPaths);
-	FTextLocalizationManager::Get().HandleLocalizationTargetsMounted(AdditionalLocResPaths);
+	if (Notification)
+	{
+		FTextLocalizationManager::Get().HandleLocalizationTargetsMounted(AdditionalLocResPaths, [PluginName, Notification = MoveTemp(Notification)]()
+		{
+			Notification(true, PluginName);
+		});
+	}
+	else
+	{
+		FTextLocalizationManager::Get().HandleLocalizationTargetsMounted(AdditionalLocResPaths);
+	}
 	return true;
 }
 
-bool FPluginManager::UnmountExplicitlyLoadedPluginLocalizationData(const FString& PluginName)
+bool FPluginManager::UnmountExplicitlyLoadedPluginLocalizationData(const FString& PluginName, TUniqueFunction<void(bool bUnloadedLocalization, const FString& PluginName)>&& Notification)
 {
 	TSharedPtr<FPlugin> Plugin = FindPluginInstance(PluginName);
 
@@ -3314,6 +3340,10 @@ bool FPluginManager::UnmountExplicitlyLoadedPluginLocalizationData(const FString
 	{
 		// Does not exist or is not mounted
 		UE_LOG(LogPluginManager, Error, TEXT("Cannot unmount plugin localization for '%s' as the plugin is unknown or not mounted. Did you forget to call MountExplicitlyLoadedPlugin?"), *PluginName);
+		if (Notification)
+		{
+			Notification(false, PluginName);
+		}
 		return false;
 	}
 
@@ -3321,6 +3351,10 @@ bool FPluginManager::UnmountExplicitlyLoadedPluginLocalizationData(const FString
 	{
 		// Not supported
 		UE_LOG(LogPluginManager, Warning, TEXT("Cannot unmount plugin localization for '%s' as the plugin isn't explicitly loaded."), *PluginName);
+		if (Notification)
+		{
+			Notification(false, PluginName);
+		}
 		return false;
 	}
 
@@ -3329,6 +3363,10 @@ bool FPluginManager::UnmountExplicitlyLoadedPluginLocalizationData(const FString
 		// Not possible to unload yet
 		UE_LOG(LogPluginManager, Verbose, TEXT("Ignoring request to unmount plugin localization for '%s' as it is still referenced (ref-count: %d -> %d)."), *PluginName, Plugin->ExplicitlyLoadedLocalizationDataMountedRefCount, Plugin->ExplicitlyLoadedLocalizationDataMountedRefCount - 1);
 		--Plugin->ExplicitlyLoadedLocalizationDataMountedRefCount;
+		if (Notification)
+		{
+			Notification(false, PluginName);
+		}
 		return false;
 	}
 
@@ -3336,6 +3374,10 @@ bool FPluginManager::UnmountExplicitlyLoadedPluginLocalizationData(const FString
 	{
 		// Already unloaded
 		UE_LOG(LogPluginManager, Verbose, TEXT("Ignoring request to unmount plugin localization for '%s' as the localization data was not mounted."), *PluginName);
+		if (Notification)
+		{
+			Notification(false, PluginName);
+		}
 		return false;
 	}
 
@@ -3343,6 +3385,10 @@ bool FPluginManager::UnmountExplicitlyLoadedPluginLocalizationData(const FString
 	{
 		// Nothing to load
 		UE_LOG(LogPluginManager, Verbose, TEXT("Ignoring request to unmount plugin localization for '%s' as the plugin has no localization targets defined."), *PluginName);
+		if (Notification)
+		{
+			Notification(false, PluginName);
+		}
 		return false;
 	}
 
@@ -3350,6 +3396,10 @@ bool FPluginManager::UnmountExplicitlyLoadedPluginLocalizationData(const FString
 	{
 		// Exiting; unloading will be handled by the engine shutdown
 		UE_LOG(LogPluginManager, Verbose, TEXT("Ignoring request to unmount plugin localization for '%s' as the engine is exiting."), *PluginName);
+		if (Notification)
+		{
+			Notification(false, PluginName);
+		}
 		return false;
 	}
 
@@ -3359,7 +3409,17 @@ bool FPluginManager::UnmountExplicitlyLoadedPluginLocalizationData(const FString
 	// Notify that additional localization data should be unloaded
 	TArray<FString> AdditionalLocResPaths;
 	PluginLocalizationUtils::GetLocalizationPathsForPlugin(*Plugin, AdditionalLocResPaths);
-	FTextLocalizationManager::Get().HandleLocalizationTargetsUnmounted(AdditionalLocResPaths);
+	if (Notification)
+	{
+		FTextLocalizationManager::Get().HandleLocalizationTargetsUnmounted(AdditionalLocResPaths, [PluginName, Notification = MoveTemp(Notification)]()
+		{
+			Notification(true, PluginName);
+		});
+	}
+	else
+	{
+		FTextLocalizationManager::Get().HandleLocalizationTargetsUnmounted(AdditionalLocResPaths);
+	}
 	return true;
 }
 
diff --git a/Engine/Source/Runtime/Projects/Private/PluginManager.h b/Engine/Source/Runtime/Projects/Private/PluginManager.h
--- a/Engine/Source/Runtime/Projects/Private/PluginManager.h
+++ b/Engine/Source/Runtime/Projects/Private/PluginManager.h
@@ -220,8 +220,8 @@ public:
 	virtual bool MountExplicitlyLoadedPlugin(const FString& PluginName) override;
 	virtual bool MountExplicitlyLoadedPlugin_FromFileName(const FString& PluginFileName) override;
 	virtual bool MountExplicitlyLoadedPlugin_FromDescriptor(const FPluginReferenceDescriptor& PluginDescriptor) override;
-	virtual bool MountExplicitlyLoadedPluginLocalizationData(const FString& PluginName) override;
-	virtual bool UnmountExplicitlyLoadedPluginLocalizationData(const FString& PluginName) override;
+	virtual bool MountExplicitlyLoadedPluginLocalizationData(const FString& PluginName, TUniqueFunction<void(bool bLoadedLocalization, const FString& PluginName)>&& Notification = nullptr) override;
+	virtual bool UnmountExplicitlyLoadedPluginLocalizationData(const FString& PluginName, TUniqueFunction<void(bool bUnloadedLocalization, const FString& PluginName)>&& Notification = nullptr) override;
 	virtual bool UnmountExplicitlyLoadedPlugin(const FString& PluginName, FText* OutReason) override { return UnmountExplicitlyLoadedPlugin(PluginName, OutReason, true); }
 	virtual bool UnmountExplicitlyLoadedPlugin(const FString& PluginName, FText* OutReason, bool bAllowUnloadCode) override;
 	virtual bool GetPluginDependencies(const FString& PluginName, TArray<FPluginReferenceDescriptor>& PluginDependencies) override;
diff --git a/Engine/Source/Runtime/Projects/Public/Interfaces/IPluginManager.h b/Engine/Source/Runtime/Projects/Public/Interfaces/IPluginManager.h
--- a/Engine/Source/Runtime/Projects/Public/Interfaces/IPluginManager.h
+++ b/Engine/Source/Runtime/Projects/Public/Interfaces/IPluginManager.h
@@ -541,17 +541,27 @@ public:
 	/**
 	 * Start loading localization data for an explicitly loaded plugin that has previously been mounted via one of the MountExplicitlyLoadedPlugin functions.
 	 * @note Localization data is ref-counted, so this should be paired with a call to UnmountExplicitlyLoadedPluginLocalizationData.
+	 * 
+	 * @param PluginName The name of the plugin to load the localization data for.
+	 * @param Notification An optional notification function to call when the localization data load has finished (from any thread), or immediately if there was no localization data to load.
+	 *            bLoadedLocalization corresponds to the return value of this function.
+	 * 
 	 * @return True if localization data started to load, or false if the plugin was missing or had no localization data to load.
 	 */
-	virtual bool MountExplicitlyLoadedPluginLocalizationData(const FString& PluginName) = 0;
+	virtual bool MountExplicitlyLoadedPluginLocalizationData(const FString& PluginName, TUniqueFunction<void(bool bLoadedLocalization, const FString& PluginName)>&& Notification = nullptr) = 0;
 
 	/**
 	 * Start unloading localization data for an explicitly loaded plugin that had its localization data mounted via MountExplicitlyLoadedPluginLocalizationData.
 	 * @note Localization data is ref-counted, so this should be paired with a call to MountExplicitlyLoadedPluginLocalizationData.
 	 * @note Localization data is also automatically (and regardless of ref-count) unloaded when calling UnmountExplicitlyLoadedPlugin.
+	 *
+	 * @param PluginName The name of the plugin to unload the localization data for.
+	 * @param Notification An optional notification function to call when the localization data unload has finished (from any thread), or immediately if there was no localization data to unload.
+	 *            bUnloadedLocalization corresponds to the return value of this function.
+	 *
 	 * @return True if localization data started to unload, or false if the plugin was missing or had no localization data to unload.
 	 */
-	virtual bool UnmountExplicitlyLoadedPluginLocalizationData(const FString& PluginName) = 0;
+	virtual bool UnmountExplicitlyLoadedPluginLocalizationData(const FString& PluginName, TUniqueFunction<void(bool bUnloadedLocalization, const FString& PluginName)>&& Notification = nullptr) = 0;
 
 	/**
 	 * Marks an explicitly loaded plugin as disabled, unmounts its content (does not work on plugins with compiled modules).
