Adding a record of properties for opened projects for associating the location of the editor that was last used

#rnx

diff --git a/Engine/Source/Editor/UnrealEd/Private/EditorEngine.cpp b/Engine/Source/Editor/UnrealEd/Private/EditorEngine.cpp
--- a/Engine/Source/Editor/UnrealEd/Private/EditorEngine.cpp
+++ b/Engine/Source/Editor/UnrealEd/Private/EditorEngine.cpp
@@ -141,6 +141,7 @@
 #include "ISourceControlWindowsModule.h"
 
 #include "Settings/EditorSettings.h"
+#include "ProjectEditorRecords.h"
 
 #include "Interfaces/IMainFrameModule.h"
 #include "LevelEditor.h"
@@ -1593,6 +1594,8 @@ void UEditorEngine::PreExit()
 		EditorSubsystemCollection.Deinitialize();
 	}
 
+	FProjectEditorRecord::TearDown();
+
 	Super::PreExit();
 }
 
@@ -5508,6 +5511,13 @@ void UEditorEngine::UpdateRecentlyLoadedProjectFiles()
 		}
 
 		GetMutableDefault<UEditorSettings>()->PostEditChange();
+		
+		FProjectEditorRecord::QueueUpdate(
+		[AbsoluteProjectPath](FProjectEditorRecord& InEditorRecords)
+		{
+			const TSharedPtr<FJsonObject>& Projects = InEditorRecords.FindOrAddProperty(FProjectEditorRecord::ProjectsProperty);
+			Projects->SetObjectField(AbsoluteProjectPath, FProjectEditorRecord::MakeDefaultProperties());
+		});
 	}
 }
 
diff --git a/Engine/Source/Runtime/Projects/Private/ProjectEditorRecords.cpp b/Engine/Source/Runtime/Projects/Private/ProjectEditorRecords.cpp
new file mode 100644
--- /dev/null
+++ b/Engine/Source/Runtime/Projects/Private/ProjectEditorRecords.cpp
@@ -0,0 +1,176 @@
+// Copyright Epic Games, Inc. All Rights Reserved.
+
+#include "ProjectEditorRecords.h"
+#include "HAL/PlatformProcess.h"
+#include "Misc/FileHelper.h"
+#include "Serialization/JsonSerializer.h"
+#include "Misc/Paths.h"
+#include "Misc/App.h"
+#include "HAL/ConsoleManager.h"
+#include "Async/TaskGraphInterfaces.h"
+
+const FString FProjectEditorRecord::ProjectsProperty = TEXT("Projects");
+const FString FProjectEditorRecord::SubProjectProperty = TEXT("SubProjects");
+const FString FProjectEditorRecord::EngineLocationProperty = TEXT("EngineLocation");
+const FString FProjectEditorRecord::BaseDirProperty = TEXT("BaseDir");
+const FString FProjectEditorRecord::TimestampProperty = TEXT("LastAccessed");
+const FString FProjectEditorRecord::EpicAppProperty = TEXT("EpicApp");
+FGraphEventRef FProjectEditorRecord::AsyncUpdateTask = nullptr;
+
+
+namespace ProjectEditorRecords
+{
+	// Lazy initialized in GetFileLocation if not set
+	FString ProjectEngineLocationFile;
+	static FAutoConsoleVariableRef CVarProjectEngineLocationFile(
+		TEXT("r.Editor.ProjectEditorRecordsFile"),
+		ProjectEngineLocationFile,
+		TEXT("The path of the Project - Engine Location record file.")
+	);
+}
+
+FProjectEditorRecord FProjectEditorRecord::Load()
+{	
+	FProjectEditorRecord Records;
+	FString FileContents;
+	if (!FFileHelper::LoadFileToString(FileContents, *GetFileLocation()))
+	{
+		Records.ProjectEditorJson = MakeShared<FJsonObject>();
+		return Records;
+	}
+
+	TSharedRef<TJsonReader<>> Reader = TJsonReaderFactory<>::Create(FileContents);
+	if (!FJsonSerializer::Deserialize(Reader, Records.ProjectEditorJson) || !Records.ProjectEditorJson.IsValid())
+	{
+		Records.ProjectEditorJson = MakeShared<FJsonObject>();
+	}
+
+	return Records;
+}
+
+bool FProjectEditorRecord::Save() 
+{
+	PruneOldEntries(ProjectEditorJson);
+	FString FileContents;
+
+	TSharedRef<TJsonWriter<>> Writer = TJsonWriterFactory<>::Create(&FileContents);
+	if (!FJsonSerializer::Serialize(ProjectEditorJson.ToSharedRef(), Writer))
+	{
+		return false;
+	}
+
+	if (!FFileHelper::SaveStringToFile(FileContents, *GetFileLocation()))
+	{
+		return false;
+	}
+
+	return true;
+}
+
+const TSharedPtr<FJsonObject>& FProjectEditorRecord::FindOrAddProperty(const FString& InProperty)
+{
+	if (!ProjectEditorJson->HasTypedField(InProperty, EJson::Object))
+	{
+		ProjectEditorJson->SetObjectField(InProperty, MakeShared<FJsonObject>());
+	}
+
+	const TSharedPtr<FJsonObject>& Field = ProjectEditorJson->GetObjectField(InProperty);
+	Field->SetStringField(FProjectEditorRecord::TimestampProperty, FDateTime::UtcNow().ToString());
+	return Field;
+
+}
+
+TSharedPtr<FJsonObject> FProjectEditorRecord::MakeDefaultProperties()
+{
+	TSharedPtr<FJsonObject> JsonProperties = MakeShared<FJsonObject>();
+	JsonProperties->SetStringField(FProjectEditorRecord::EngineLocationProperty, FPlatformProcess::ExecutablePath());
+	JsonProperties->SetStringField(FProjectEditorRecord::BaseDirProperty, FPlatformProcess::BaseDir());
+	JsonProperties->SetStringField(FProjectEditorRecord::TimestampProperty, FDateTime::UtcNow().ToString());
+	return JsonProperties;
+}
+
+void FProjectEditorRecord::QueueUpdate(TUniqueFunction<void(FProjectEditorRecord&)>&& InUpdateFunction)
+{
+	TUniqueFunction<void()> Task = [Update = MoveTemp(InUpdateFunction)]()
+	{
+		FSystemWideCriticalSection SystemMutex(TEXT("ProjectEditorRecords"), FTimespan::FromMinutes(1));
+		FProjectEditorRecord AssociationFile = FProjectEditorRecord::Load();
+		Update(AssociationFile);
+		AssociationFile.Save();
+	};
+
+	if (FTaskGraphInterface::IsRunning())
+	{
+		if (AsyncUpdateTask)
+		{
+			AsyncUpdateTask = FFunctionGraphTask::CreateAndDispatchWhenReady(MoveTemp(Task), TStatId(), AsyncUpdateTask);
+		}
+		else
+		{
+			AsyncUpdateTask = FFunctionGraphTask::CreateAndDispatchWhenReady(MoveTemp(Task), TStatId());
+		}
+	}
+	else
+	{
+		Task();
+	}
+}
+ 
+void FProjectEditorRecord::TearDown()
+{
+	if (AsyncUpdateTask)
+	{
+		if (FTaskGraphInterface::IsRunning())
+		{
+			FTaskGraphInterface::Get().WaitUntilTaskCompletes(AsyncUpdateTask);
+		}
+		AsyncUpdateTask.SafeRelease();
+	}
+}
+
+void FProjectEditorRecord::PruneOldEntries(const TSharedPtr<FJsonObject>& JsonObject)
+{
+	TArray<FString> KeysToRemove;
+	for (const TPair<FString, TSharedPtr<FJsonValue>>& Field : JsonObject->Values)
+	{
+		// Non-object fields are left alone, top level should always be objects with a LastAccessed field that we can check
+		if (Field.Value->Type != EJson::Object)
+		{
+			continue;
+		}
+
+		TSharedPtr<FJsonObject>* ProjectInfo;
+		FString TimestampStr;
+		FDateTime Timestamp;
+
+		if (!Field.Value->TryGetObject(ProjectInfo) 
+		|| !(*ProjectInfo)->TryGetStringField(FProjectEditorRecord::TimestampProperty, TimestampStr) 
+		|| !FDateTime::Parse(TimestampStr, Timestamp) 
+		|| (Timestamp - FDateTime::UtcNow()).GetTotalDays() > DaysToKeepRecords)
+		{
+			// If it's not valid, mark it for pruning
+			KeysToRemove.Add(Field.Key);
+		}
+		else
+		{
+			// If it's valid, search child objects recursiveness for pruning
+			PruneOldEntries(Field.Value->AsObject());
+		}
+	}
+
+	for (const FString& Key : KeysToRemove)
+	{
+		ProjectEditorJson->RemoveField(Key);
+	}
+}
+
+const FString FProjectEditorRecord::GetFileLocation()
+{
+	FString FileLocation = ProjectEditorRecords::CVarProjectEngineLocationFile->GetString();
+	if (FileLocation.IsEmpty())
+	{
+		FileLocation = FPaths::Combine(FPlatformProcess::UserSettingsDir(), *FApp::GetEpicProductIdentifier(), TEXT("Editor"), "ProjectEditorRecords.json");
+	}
+
+	return FileLocation;
+}
diff --git a/Engine/Source/Runtime/Projects/Public/ProjectEditorRecords.h b/Engine/Source/Runtime/Projects/Public/ProjectEditorRecords.h
new file mode 100644
--- /dev/null
+++ b/Engine/Source/Runtime/Projects/Public/ProjectEditorRecords.h
@@ -0,0 +1,105 @@
+// Copyright Epic Games, Inc. All Rights Reserved.
+
+#pragma once
+
+#include "Dom/JsonObject.h"
+
+/**
+ * Struct specifying the structure of the records file, this is used to maintain
+ * an association between recently opened projects and the location of the editor engine
+ * used to open them. This is meant to be a machine wide record to facilitate locating 
+ * the right engine executable.
+ * 
+ * Usage: 
+ * - QueueUpdate to update the record
+ * - Load to read from the record
+ */
+struct FProjectEditorRecord
+{
+	/** Property used to specify a Projects list */
+	PROJECTS_API static const FString ProjectsProperty;
+
+	/** Property used to specify a SubProjects list */
+	PROJECTS_API static const FString SubProjectProperty;
+
+	/** Property used to specify the path of the engine used to open the project */
+	PROJECTS_API static const FString EngineLocationProperty;
+
+	/** Property used to specify the path of the base dir of the engine used for the project */
+	PROJECTS_API static const FString BaseDirProperty;
+	
+	/** Property used to specify a last accessed timestamp */
+	PROJECTS_API static const FString TimestampProperty;
+
+	/** Property used to specify an Epic App */
+	PROJECTS_API static const FString EpicAppProperty;
+
+	/** The json contents of the records loaded/saved to disk */
+	TSharedPtr<FJsonObject> ProjectEditorJson;
+
+	/**
+	 * This is the main way to interact with the json contents to add entries in a safe way.
+	 *
+	 * QueueUpdate schedules a task that will try to run in another thread and perform these operations:
+	 * - Acquires a system wide lock
+	 * - Loads the latest records from disk
+	 * - Runs the function passed as a parameter with the loaded contents
+	 * - Saves the records to disk
+	 * - Releases system wide lock
+	 *
+	 * @param Function used to update the records in a worker thread
+	 */
+	PROJECTS_API static void QueueUpdate(TUniqueFunction<void(FProjectEditorRecord&)>&& InUpdateFunction);
+
+	/**
+	 * Loads a project editor record from disk, using the default location %LOCALAPPDATA%/UnrealEngine/Editor/ProjectEditorRecords.json
+	 * location can be overriden by cvar r.Editor.ProjectEditorRecordsFile
+	 * 
+	 * @return	The struct loaded from disk, or an empty one.
+	 */
+	PROJECTS_API static FProjectEditorRecord Load();
+	
+	/**
+	 * Finds or adds an object property to the root of the json
+	 * 
+	 * @param The name of the property to find or add
+	 * @return Shared pointer to the json object corresponding to the property input
+	 */
+	PROJECTS_API const TSharedPtr<FJsonObject>& FindOrAddProperty(const FString& InProperty);
+
+	/**
+	 * Creates a JsonObject containing the default entries 
+	 * - Path to the UnrealEditor.exe used when calling this function (meant to be UnrealEditor.exe)
+	 * - Path to the base directory when calling this function
+	 * - Current timestamp
+	 * 
+	 * @return JsonObject containing the default entries for a project
+	 */
+	PROJECTS_API static TSharedPtr<FJsonObject> MakeDefaultProperties();
+
+
+	/**
+	 * Waits for any queued async tasks to finish if there are any
+	 */
+	PROJECTS_API static void TearDown();
+
+private:
+	FProjectEditorRecord() = default;
+
+	/** Handle for the latest async task */
+	static FGraphEventRef AsyncUpdateTask;
+
+	/** Number of days that have to pass for an entry to be considered stale and be removed on save */
+	const int32 DaysToKeepRecords = 60;
+
+	/** Goes over the json object and prunes anything older than DaysToKeepRecords */
+	void PruneOldEntries(const TSharedPtr<FJsonObject>& JsonObject);
+
+	/** Gets the location of the file on disk */
+	static const FString GetFileLocation();
+
+	/** Saves the contents of this structure to disk */
+	bool Save();
+};
+
+
