Downgrade a couple of logs to verbose

diff --git a/Engine/Source/Runtime/Projects/Private/PluginManager.cpp b/Engine/Source/Runtime/Projects/Private/PluginManager.cpp
--- a/Engine/Source/Runtime/Projects/Private/PluginManager.cpp
+++ b/Engine/Source/Runtime/Projects/Private/PluginManager.cpp
@@ -2492,7 +2492,7 @@ bool FPluginManager::ConfigureEnabledPluginForTarget(const FPluginReferenceDescr
 #if !UE_BUILD_SHIPPING
 		if (Reference.bOptional && !AllowedOptionalDependencies.IsEmpty() && !AllowedOptionalDependencies.Contains(Reference.Name))
 		{
-			UE_LOG(LogPluginManager, Display, TEXT("Ignored optional reference to '%s' plugin; plugin was not built by target."), *Reference.Name);
+			UE_LOG(LogPluginManager, Verbose, TEXT("Ignored optional reference to '%s' plugin; plugin was not built by target."), *Reference.Name);
 			continue;
 		}
 #endif
@@ -2531,7 +2531,7 @@ bool FPluginManager::ConfigureEnabledPluginForTarget(const FPluginReferenceDescr
 			// Allowed optional plugins are compiled enabled or enabled via the commandline. Ignore those that are not.
 			if (NextReference.bOptional && !AllowedOptionalDependencies.IsEmpty() && !AllowedOptionalDependencies.Contains(NextReference.Name))
 			{
-				UE_LOG(LogPluginManager, Display, TEXT("Ignored optional reference to '%s' plugin from '%s' plugin; plugin was not built by target."), *NextReference.Name, *Plugin.GetName());
+				UE_LOG(LogPluginManager, Verbose, TEXT("Ignored optional reference to '%s' plugin from '%s' plugin; plugin was not built by target."), *NextReference.Name, *Plugin.GetName());
 				continue;
 			}
 #endif
