Update UBT to report when a project uses a deprecated plugin and promote the field to the PluginManager so users can see when a plugin is mounted that has been marked deprecated.

diff --git a/Engine/Source/Programs/UnrealBuildTool/Configuration/UEBuildPlugin.cs b/Engine/Source/Programs/UnrealBuildTool/Configuration/UEBuildPlugin.cs
--- a/Engine/Source/Programs/UnrealBuildTool/Configuration/UEBuildPlugin.cs
+++ b/Engine/Source/Programs/UnrealBuildTool/Configuration/UEBuildPlugin.cs
@@ -123,6 +123,14 @@ namespace UnrealBuildTool
 				}
 			}
 
+			if (!String.IsNullOrEmpty(Descriptor.DeprecatedEngineVersion))
+			{
+				// Eventually we need to promote this to an actual warning but there are a number of uses of deprecated plugins that would cause a lot
+				// of noise and we want to gently nudge people to migrate plugins.
+				logger.LogInformation("Project depends on plugin '{PluginName}' which was deprecated in {EngineVersion} and will soon be removed. Please update your dependencies.",
+									  Name, Descriptor.DeprecatedEngineVersion);
+			}
+			
 			// Check that any plugins with the NoCode specifier do not contain modules.
 			if (Descriptor.bNoCode && Modules.Any())
 			{
diff --git a/Engine/Source/Runtime/Projects/Private/PluginDescriptor.cpp b/Engine/Source/Runtime/Projects/Private/PluginDescriptor.cpp
--- a/Engine/Source/Runtime/Projects/Private/PluginDescriptor.cpp
+++ b/Engine/Source/Runtime/Projects/Private/PluginDescriptor.cpp
@@ -221,6 +221,7 @@ bool FPluginDescriptor::Read(const FJsonObject& Object, FText* OutFailReason /*=
 	}
 
 	Object.TryGetStringField(TEXT("VersePath"), VersePath);
+	Object.TryGetStringField(TEXT("DeprecatedEngineVersion"), DeprecatedEngineVersion);
 
 	// Read the Verse scope
 	TSharedPtr<FJsonValue> VerseScopeValue = Object.TryGetField(TEXT("VerseScope"));
diff --git a/Engine/Source/Runtime/Projects/Private/PluginManager.cpp b/Engine/Source/Runtime/Projects/Private/PluginManager.cpp
--- a/Engine/Source/Runtime/Projects/Private/PluginManager.cpp
+++ b/Engine/Source/Runtime/Projects/Private/PluginManager.cpp
@@ -53,6 +53,14 @@ void SetCoreUObjectPluginManager(ICoreUObjectPluginManager& Handler)
 {
 	CoreUObjectPluginHandler = &Handler;
 }
+
+void CheckForDeprecatedPlugin(FPlugin& Plugin)
+{
+	if (const FString& DeprecatedEngineVersion = Plugin.GetDeprecatedEngineVersion(); !DeprecatedEngineVersion.IsEmpty())
+	{
+		UE_LOG(LogPluginManager, Display, TEXT("The Plugin %s has been marked deprecated for %s and will be removed soon. Please update your dependencies."), *Plugin.GetName(), *DeprecatedEngineVersion);
+	}
+}
 	
 TArray<FString> GetPluginPathsByEnv( const TCHAR* EnvVariable )
 {
@@ -1924,7 +1932,9 @@ bool FPluginManager::ConfigureEnabledPlugins()
 				
 				UE_LOG(LogPluginManager, Log, TEXT("Mounting %s plugin %s"), *EnumToString(Plugin.Type), *Plugin.GetName());
 				UE_LOG(LogPluginManager, Verbose, TEXT("Plugin path: %s"), *Plugin.FileName);
-				
+
+				UE::PluginManager::Private::CheckForDeprecatedPlugin(Plugin);
+
 				// register this plugin, so it can be found right now, but also when loading other platforms, or when modifying it by other plugins
 				const bool bIncludePluginNameInBranchName = false;
 				FConfigCacheIni::RegisterPlugin(PluginName, Plugin.GetBaseDir(), Plugin.GetExtensionBaseDirs(), DynamicLayerPriority::Plugin, bIncludePluginNameInBranchName);
@@ -3468,6 +3478,8 @@ void FPluginManager::MountPluginFromExternalSource(const TSharedRef<FPlugin>& Pl
 	Plugin->bEnabled = true;
 	PluginsEnabledForTarget.Add(Plugin->GetName());
 
+	UE::PluginManager::Private::CheckForDeprecatedPlugin(*Plugin);
+	
 	// Mount the plugin content directory
 	const bool bHasContentOrVerse = (Plugin->CanContainContent() || Plugin->CanContainVerse()) && ensure(RegisterMountPointDelegate.IsBound());
 	if (bHasContentOrVerse)
diff --git a/Engine/Source/Runtime/Projects/Private/PluginManager.h b/Engine/Source/Runtime/Projects/Private/PluginManager.h
--- a/Engine/Source/Runtime/Projects/Private/PluginManager.h
+++ b/Engine/Source/Runtime/Projects/Private/PluginManager.h
@@ -116,6 +116,11 @@ public:
 		return Descriptor.VerseVersion;
 	}
 
+	virtual const FString& GetDeprecatedEngineVersion() const override
+	{
+		return Descriptor.DeprecatedEngineVersion;
+	}
+		
 	virtual bool IsSceneGraphEnabled() const override
 	{
 		return Descriptor.bEnableSceneGraph;
diff --git a/Engine/Source/Runtime/Projects/Public/Interfaces/IPluginManager.h b/Engine/Source/Runtime/Projects/Public/Interfaces/IPluginManager.h
--- a/Engine/Source/Runtime/Projects/Public/Interfaces/IPluginManager.h
+++ b/Engine/Source/Runtime/Projects/Public/Interfaces/IPluginManager.h
@@ -216,6 +216,14 @@ public:
 	 */
 	virtual const FString& GetVersePath() const = 0;
 
+	/**
+	 * Gets the engine version that this plugin was deprecated.  If unset
+	 * then this plugin is not deprecated.
+	 *
+	 * @return Deprecated engine version.  This may be an empty string.
+	 */
+	virtual const FString& GetDeprecatedEngineVersion() const = 0;
+
 	/**
 	 * Return the Verse language version to compile the plugin with (if unspecified, the latest stable version is used)
 	 */
diff --git a/Engine/Source/Runtime/Projects/Public/PluginDescriptor.h b/Engine/Source/Runtime/Projects/Public/PluginDescriptor.h
--- a/Engine/Source/Runtime/Projects/Public/PluginDescriptor.h
+++ b/Engine/Source/Runtime/Projects/Public/PluginDescriptor.h
@@ -94,6 +94,9 @@ struct FPluginDescriptor
 	/** The Verse path to the root of this plugin's content directory */
 	FString VersePath;
 
+	/** Set when a plugin is loaded that has the DeprecatedEngineVersion set. */
+	FString DeprecatedEngineVersion; 
+
 	/** Origin/visibility of Verse code in this plugin's Content/Verse folder */
 	EVerseScope::Type VerseScope = EVerseScope::PublicUser;
 
