Add activate syntax to plugins as this is required by the GameFeaturePlugin manager.
Having this property missing can result in these values going missing when loading and writng out a PluginDescriptor.

diff --git a/Engine/Source/Runtime/Projects/Private/PluginReferenceDescriptor.cpp b/Engine/Source/Runtime/Projects/Private/PluginReferenceDescriptor.cpp
--- a/Engine/Source/Runtime/Projects/Private/PluginReferenceDescriptor.cpp
+++ b/Engine/Source/Runtime/Projects/Private/PluginReferenceDescriptor.cpp
@@ -32,6 +32,7 @@ FPluginReferenceDescriptor::FPluginReferenceDescriptor( const FString& InName, b
 	: Name(InName)
 	, bEnabled(bInEnabled)
 	, bOptional(false)
+	, bActivate(false)
 	, bHasExplicitPlatforms(false)
 { }
 
@@ -156,6 +157,8 @@ bool FPluginReferenceDescriptor::Read(const FJsonObject& Object, FText* OutFailR
 	// Read the optional field
 	Object.TryGetBoolField(TEXT("Optional"), bOptional);
 
+	Object.TryGetBoolField(TEXT("Activate"), bActivate);
+
 	// Read the metadata for users that don't have the plugin installed
 	Object.TryGetStringField(TEXT("Description"), Description);
 	Object.TryGetStringField(TEXT("MarketplaceURL"), MarketplaceURL);
@@ -253,6 +256,15 @@ void FPluginReferenceDescriptor::UpdateJson(FJsonObject& JsonObject) const
 	JsonObject.SetStringField(TEXT("Name"), Name);
 	JsonObject.SetBoolField(TEXT("Enabled"), bEnabled);
 
+	if (bActivate)
+	{
+		JsonObject.SetBoolField(TEXT("Activate"), bActivate);
+	}
+	else
+	{
+		JsonObject.RemoveField(TEXT("Activate"));
+	}
+
 	if (bEnabled && bOptional)
 	{
 		JsonObject.SetBoolField(TEXT("Optional"), bOptional);
diff --git a/Engine/Source/Runtime/Projects/Public/PluginReferenceDescriptor.h b/Engine/Source/Runtime/Projects/Public/PluginReferenceDescriptor.h
--- a/Engine/Source/Runtime/Projects/Public/PluginReferenceDescriptor.h
+++ b/Engine/Source/Runtime/Projects/Public/PluginReferenceDescriptor.h
@@ -33,6 +33,9 @@ struct FPluginReferenceDescriptor
 
 	/** Whether this plugin is optional, and the game should silently ignore it not being present */
 	bool bOptional;
+
+	/** Whether this plugin dependency should also activate when the parent plugin is activated by the GameFeaturePlugin system. */
+	bool bActivate;
 	
 	/** Description of the plugin for users that do not have it installed. */
 	FString Description;
