Fix Plugin subdir recursion going into Content dirs could cause very long stall at startup

diff --git a/Engine/Source/Runtime/Projects/Private/PluginManager.cpp b/Engine/Source/Runtime/Projects/Private/PluginManager.cpp
--- a/Engine/Source/Runtime/Projects/Private/PluginManager.cpp
+++ b/Engine/Source/Runtime/Projects/Private/PluginManager.cpp
@@ -1108,16 +1108,28 @@ void FPluginManager::FindPluginsInDirectory(const FString& PluginsDirectory, TAr
 
 		virtual bool Visit(const TCHAR* FilenameOrDirectory, bool bIsDirectory) override
 		{
+			int32 FileNameLen = FCString::Strlen(FilenameOrDirectory);
+
 			if (bIsDirectory)
 			{
-				SubDirectories.Emplace(FilenameOrDirectory);
+				constexpr const TCHAR ContentStr[] = TEXT("/Content");
+				constexpr const int32 ContentStrLen = UE_ARRAY_COUNT(ContentStr)-1; // -1 for the string's terminating zero (which would be counted by `UE_ARRAY_COUNT()`)
+
+				if (FileNameLen >= ContentStrLen && FCString::Stricmp(&FilenameOrDirectory[FileNameLen - ContentStrLen], ContentStr) == 0)
+				{
+					// excluded
+				}
+				else
+				{
+					SubDirectories.Emplace(FilenameOrDirectory);
+				}
 			}
 			else
 			{
+				// ideally we'd like to find this .uplugin file first before visiting anything else
 				constexpr const TCHAR PluginExt[] = TEXT(".uplugin");
 				constexpr const int32 PluginExtLen = UE_ARRAY_COUNT(PluginExt)-1; // -1 for the string's terminating zero (which would be counted by `UE_ARRAY_COUNT()`)
 
-				int32 FileNameLen = FCString::Strlen(FilenameOrDirectory);
 				if (FileNameLen >= PluginExtLen && FCString::Strcmp(&FilenameOrDirectory[FileNameLen - PluginExtLen], PluginExt) == 0)
 				{
 					FoundPluginFile = FilenameOrDirectory;
