Extracting a portion of ConfigureEnabledPlugins to reduce function size.
Removed BinaryConfigPlugins code path that was no longer needed.

#rnx

diff --git a/Engine/Source/Runtime/Projects/Private/PluginManager.cpp b/Engine/Source/Runtime/Projects/Private/PluginManager.cpp
--- a/Engine/Source/Runtime/Projects/Private/PluginManager.cpp
+++ b/Engine/Source/Runtime/Projects/Private/PluginManager.cpp
@@ -1362,6 +1362,108 @@ bool FPluginManager::GetPreloadBinaries()
 	return bPreloadedBinaries;
 }
 
+void FPluginManager::ProcessEnabledPlugins()
+{
+	SCOPED_BOOT_TIMING("ProcessEnabledPlugins");
+
+	// generate optimal list of plugins to process
+	TArray<TSharedRef<FPlugin>> PluginsArray;
+	for (const FString& PluginName : PluginsToConfigure)
+	{
+		TSharedRef<FPlugin> Plugin = DiscoveredPluginMapUtils::FindPluginInMap_Checked(AllPlugins, PluginName);
+		
+		// only process enabled plugins
+		if (Plugin->bEnabled && !Plugin->Descriptor.bExplicitlyLoaded)
+		{
+			PluginsArray.Add(Plugin);
+		}
+	}
+
+	TArray<FString> ConfigFilesPluginsCannotOverride;
+	GConfig->GetArray(TEXT("Plugins"), TEXT("ConfigFilesPluginsCannotOverride"), ConfigFilesPluginsCannotOverride, GEngineIni);
+	
+	FCriticalSection ConfigCS;
+	FCriticalSection PluginPakCS;
+	
+	// Mount all the enabled plugins
+	FConfigModificationTracker ChangeTracker;
+	ChangeTracker.bTrackLoadedFiles = true;
+
+	TArray<FName> PluginNames;
+	PluginNames.Reserve(PluginsArray.Num());
+	
+	// walk over each plugin, and add the config files that are named for the plugin
+	// this is a separate loop so that a plugin can modify another plugin's configs in the second loop below
+	for (TSharedRef<FPlugin> PluginPtr : PluginsArray)
+	{
+		FPlugin& Plugin = *PluginPtr;
+		FName PluginName(*Plugin.Name);
+		PluginNames.Add(PluginName);
+		
+		UE_LOG(LogPluginManager, Log, TEXT("Mounting %s plugin %s"), *EnumToString(Plugin.Type), *Plugin.GetName());
+		UE_LOG(LogPluginManager, Verbose, TEXT("Plugin path: %s"), *Plugin.FileName);
+
+		UE::PluginManager::Private::CheckForDeprecatedPlugin(Plugin);
+
+		// register this plugin, so it can be found right now, but also when loading other platforms, or when modifying it by other plugins
+		const bool bIncludePluginNameInBranchName = false;
+		FConfigCacheIni::RegisterPlugin(PluginName, Plugin.GetBaseDir(), Plugin.GetExtensionBaseDirs(), DynamicLayerPriority::Plugin, bIncludePluginNameInBranchName);
+		
+		FConfigContext Context = FConfigContext::ReadIntoGConfig();
+		Context.ChangeTracker = &ChangeTracker;
+		Context.ConfigFileTag = *Plugin.Name;
+		Context.Load(*Plugin.Name);
+	}
+	
+	// gang up dynamic layer adds as much as possible to reduce the amount of "CalcDiff" calls etc 
+	FConfigCacheIni::AddMultiplePluginsToAllBranches(PluginNames);
+
+	// walk over each plugin, modify existing ini's, mount the content, etc
+	ParallelFor(PluginsArray.Num(), [&PluginsArray, &ConfigCS, &PluginPakCS, &ConfigFilesPluginsCannotOverride, &ChangeTracker, this](int32 Index)
+	{
+		FPlugin& Plugin = *PluginsArray[Index];
+		FName PluginName(*Plugin.Name);
+
+		// Build the list of content folders
+		if (Plugin.Descriptor.bCanContainContent)
+		{
+			{
+				FScopeLock Locker(&ConfigCS);
+
+				// we need to look up the section each time because other loops could add entries
+				if (FConfigFile* EngineConfigFile = GConfig->FindConfigFile(GEngineIni))
+				{
+					EngineConfigFile->AddUniqueToSection(TEXT("Core.System"), "Paths", Plugin.GetContentDir());
+				}
+			}
+
+			TArray<FString>	FoundPaks;
+			FPakFileSearchVisitor PakVisitor(FoundPaks);
+			IPlatformFile& PlatformFile = FPlatformFileManager::Get().GetPlatformFile();
+
+			// Pak files are loaded from <PluginName>/Content/Paks/<PlatformName>
+			if (FPlatformProperties::RequiresCookedData())
+			{
+				PlatformFile.IterateDirectoryRecursively(*(Plugin.GetContentDir() / TEXT("Paks") / FPlatformProperties::PlatformName()), PakVisitor);
+
+				for (const FString& PakPath : FoundPaks)
+				{
+					FScopeLock Locker(&PluginPakCS);
+					if (FCoreDelegates::MountPak.IsBound())
+					{
+						FCoreDelegates::MountPak.Execute(PakPath, 0);
+						PluginsWithPakFile.AddUnique(PluginsArray[Index]);
+					}
+					else
+					{
+						UE_LOG(LogPluginManager, Warning, TEXT("PAK file (%s) could not be mounted because MountPak is not bound"), *PakPath);
+					}
+				}
+			}
+		}
+	}, true); // @todo disable parallelism for now as it's causing hard to track problems
+}
+
 bool FPluginManager::ConfigureEnabledPlugins()
 {
 #if (WITH_ENGINE && !IS_PROGRAM) || WITH_PLUGIN_SUPPORT
@@ -1888,148 +1990,7 @@ bool FPluginManager::ConfigureEnabledPlugins()
 		// If we made it here, we have all the required plugins
 		bHaveAllRequiredPlugins = true;
 
-		// check if the config already contains the plugin inis - if so, we don't need to scan anything, just use the ini to find paks to mount
-		TArray<FString> BinaryConfigPlugins;
-		if (GConfig->GetArray(TEXT("BinaryConfig"), TEXT("BinaryConfigPlugins"), BinaryConfigPlugins, GEngineIni) && BinaryConfigPlugins.Num() > 0)
-		{
-			SCOPED_BOOT_TIMING("QuickMountingPaks");
-
-			TArray<FString> PluginPaks;
-			GConfig->GetArray(TEXT("Core.System"), TEXT("PluginPaks"), PluginPaks, GEngineIni);
-			if (FCoreDelegates::MountPak.IsBound())
-			{
-				for (FString& PakPathEntry : PluginPaks)
-				{
-					int32 PipeLocation;
-					if (PakPathEntry.FindChar(TEXT('|'), PipeLocation))
-					{
-						// split the string in twain
-						FString PluginName = PakPathEntry.Left(PipeLocation);
-						FString PakPath = PakPathEntry.Mid(PipeLocation + 1);
-
-						// look for the existing plugin
-						FPlugin* FoundPlugin = EnabledPlugins.FindRef(PluginName);
-						if (FoundPlugin != nullptr)
-						{
-							PluginsWithPakFile.AddUnique(TSharedRef<IPlugin>(FoundPlugin));
-							// and finally mount the plugin's pak
-							FCoreDelegates::MountPak.Execute(PakPath, 0);
-						}
-						
-					}
-				}
-			}
-			else
-			{
-				UE_LOG(LogPluginManager, Warning, TEXT("Plugin Pak files could not be mounted because MountPak is not bound"));
-			}
-		}
-
-
-		// even if we had plugins in the Config already, we need to process Mod plugins
-		{
-			SCOPED_BOOT_TIMING("ParallelPluginEnabling");
-
-			// generate optimal list of plugins to process
-			TArray<TSharedRef<FPlugin>> PluginsArray;
-			for (const FString& PluginName : PluginsToConfigure)
-			{
-				TSharedRef<FPlugin> Plugin = DiscoveredPluginMapUtils::FindPluginInMap_Checked(AllPlugins, PluginName);
-				// check all plugins that were not in a BinaryConfig
-				if (!BinaryConfigPlugins.Contains(PluginName))
-				{
-					// only process enabled plugins
-					if (Plugin->bEnabled && !Plugin->Descriptor.bExplicitlyLoaded)
-					{
-						PluginsArray.Add(Plugin);
-					}
-				}
-			}
-
-			TArray<FString> ConfigFilesPluginsCannotOverride;
-			GConfig->GetArray(TEXT("Plugins"), TEXT("ConfigFilesPluginsCannotOverride"), ConfigFilesPluginsCannotOverride, GEngineIni);
-			
-			FCriticalSection ConfigCS;
-			FCriticalSection PluginPakCS;
-			
-			// Mount all the enabled plugins
-			FConfigModificationTracker ChangeTracker;
-			ChangeTracker.bTrackLoadedFiles = true;
-
-			TArray<FName> PluginNames;
-			PluginNames.Reserve(PluginsArray.Num());
-			
-			// walk over each plugin, and add the config files that are named for the plugin
-			// this is a separate loop so that a plugin can modify another plugin's configs in the second loop below
-			for (TSharedRef<FPlugin> PluginPtr : PluginsArray)
-			{
-				FPlugin& Plugin = *PluginPtr;
-				FName PluginName(*Plugin.Name);
-				PluginNames.Add(PluginName);
-				
-				UE_LOG(LogPluginManager, Log, TEXT("Mounting %s plugin %s"), *EnumToString(Plugin.Type), *Plugin.GetName());
-				UE_LOG(LogPluginManager, Verbose, TEXT("Plugin path: %s"), *Plugin.FileName);
-
-				UE::PluginManager::Private::CheckForDeprecatedPlugin(Plugin);
-
-				// register this plugin, so it can be found right now, but also when loading other platforms, or when modifying it by other plugins
-				const bool bIncludePluginNameInBranchName = false;
-				FConfigCacheIni::RegisterPlugin(PluginName, Plugin.GetBaseDir(), Plugin.GetExtensionBaseDirs(), DynamicLayerPriority::Plugin, bIncludePluginNameInBranchName);
-				
-				FConfigContext Context = FConfigContext::ReadIntoGConfig();
-				Context.ChangeTracker = &ChangeTracker;
-				Context.ConfigFileTag = *Plugin.Name;
-				Context.Load(*Plugin.Name);
-			}
-			
-			// gang up dynamic layer adds as much as possible to reduce the amount of "CalcDiff" calls etc 
-			FConfigCacheIni::AddMultiplePluginsToAllBranches(PluginNames);
-
-			// walk over each plugin, modify existing ini's, mount the content, etc
-			ParallelFor(PluginsArray.Num(), [&PluginsArray, &ConfigCS, &PluginPakCS, &ConfigFilesPluginsCannotOverride, &ChangeTracker, this](int32 Index)
-			{
-				FPlugin& Plugin = *PluginsArray[Index];
-				FName PluginName(*Plugin.Name);
-
-				// Build the list of content folders
-				if (Plugin.Descriptor.bCanContainContent)
-				{
-					{
-						FScopeLock Locker(&ConfigCS);
-
-						// we need to look up the section each time because other loops could add entries
-						if (FConfigFile* EngineConfigFile = GConfig->FindConfigFile(GEngineIni))
-						{
-							EngineConfigFile->AddUniqueToSection(TEXT("Core.System"), "Paths", Plugin.GetContentDir());
-						}
-					}
-
-					TArray<FString>	FoundPaks;
-					FPakFileSearchVisitor PakVisitor(FoundPaks);
-					IPlatformFile& PlatformFile = FPlatformFileManager::Get().GetPlatformFile();
-
-					// Pak files are loaded from <PluginName>/Content/Paks/<PlatformName>
-					if (FPlatformProperties::RequiresCookedData())
-					{
-						PlatformFile.IterateDirectoryRecursively(*(Plugin.GetContentDir() / TEXT("Paks") / FPlatformProperties::PlatformName()), PakVisitor);
-
-						for (const FString& PakPath : FoundPaks)
-						{
-							FScopeLock Locker(&PluginPakCS);
-							if (FCoreDelegates::MountPak.IsBound())
-							{
-								FCoreDelegates::MountPak.Execute(PakPath, 0);
-								PluginsWithPakFile.AddUnique(PluginsArray[Index]);
-							}
-							else
-							{
-								UE_LOG(LogPluginManager, Warning, TEXT("PAK file (%s) could not be mounted because MountPak is not bound"), *PakPath);
-							}
-						}
-					}
-				}
-			}, true); // @todo disable parallelism for now as it's causing hard to track problems
-		}
+		ProcessEnabledPlugins();
 
 		// Mount enabled plugins with content
 		for (const FDiscoveredPluginMap::ElementType& PluginPair : AllPlugins)
diff --git a/Engine/Source/Runtime/Projects/Private/PluginManager.h b/Engine/Source/Runtime/Projects/Private/PluginManager.h
--- a/Engine/Source/Runtime/Projects/Private/PluginManager.h
+++ b/Engine/Source/Runtime/Projects/Private/PluginManager.h
@@ -250,6 +250,8 @@ private:
 	/** Searches for all plugins on disk and builds up the array of plugin objects.  Doesn't load any plugins. 
 	    This is called when the plugin manager singleton is first accessed. */
 	void DiscoverAllPlugins();
+	
+	void ProcessEnabledPlugins();
 
 	/** Reads all the plugin descriptors */
 	static void ReadAllPlugins(FDiscoveredPluginMap& Plugins, const TSet<FString>& ExtraSearchPaths,
