Only update ProjectEditorRecords when using an Unreal Editor and we're running a commandlet. This should filter out incorrect ExecutablePath() for non-editor project executables

#rnx

diff --git a/Engine/Source/Editor/UnrealEd/Private/EditorEngine.cpp b/Engine/Source/Editor/UnrealEd/Private/EditorEngine.cpp
--- a/Engine/Source/Editor/UnrealEd/Private/EditorEngine.cpp
+++ b/Engine/Source/Editor/UnrealEd/Private/EditorEngine.cpp
@@ -5531,12 +5531,19 @@ void UEditorEngine::UpdateRecentlyLoadedProjectFiles()
 
 		GetMutableDefault<UEditorSettings>()->PostEditChange();
 		
-		FProjectEditorRecord::QueueUpdate(
-		[AbsoluteProjectPath](FProjectEditorRecord& InEditorRecords)
+		if (GIsEditor && !IsRunningCommandlet())
 		{
-			const TSharedPtr<FJsonObject>& Projects = InEditorRecords.FindOrAddProperty(FProjectEditorRecord::ProjectsProperty);
-			Projects->SetObjectField(AbsoluteProjectPath, FProjectEditorRecord::MakeDefaultProperties());
-		});
+			FProjectEditorRecord::QueueUpdate(
+			[AbsoluteProjectPath](FProjectEditorRecord& InEditorRecords)
+			{
+				const TSharedPtr<FJsonObject> Projects = InEditorRecords.FindOrAddProperty(FProjectEditorRecord::ProjectsProperty);
+
+				if(Projects.IsValid())
+				{
+					Projects->SetObjectField(AbsoluteProjectPath, FProjectEditorRecord::MakeDefaultProperties());
+				}
+			});
+		}
 	}
 }
 
diff --git a/Engine/Source/Runtime/Projects/Private/ProjectEditorRecords.cpp b/Engine/Source/Runtime/Projects/Private/ProjectEditorRecords.cpp
--- a/Engine/Source/Runtime/Projects/Private/ProjectEditorRecords.cpp
+++ b/Engine/Source/Runtime/Projects/Private/ProjectEditorRecords.cpp
@@ -67,17 +67,21 @@ bool FProjectEditorRecord::Save()
 	return true;
 }
 
-const TSharedPtr<FJsonObject>& FProjectEditorRecord::FindOrAddProperty(const FString& InProperty)
+const TSharedPtr<FJsonObject> FProjectEditorRecord::FindOrAddProperty(const FString& InProperty)
 {
-	if (!ProjectEditorJson->HasTypedField(InProperty, EJson::Object))
+	if (GIsEditor && !IsRunningCommandlet())
 	{
-		ProjectEditorJson->SetObjectField(InProperty, MakeShared<FJsonObject>());
-	}
+		if (!ProjectEditorJson->HasTypedField(InProperty, EJson::Object))
+		{
+			ProjectEditorJson->SetObjectField(InProperty, MakeShared<FJsonObject>());
+		}
 
-	const TSharedPtr<FJsonObject>& Field = ProjectEditorJson->GetObjectField(InProperty);
-	Field->SetStringField(FProjectEditorRecord::TimestampProperty, FDateTime::UtcNow().ToString());
-	return Field;
+		const TSharedPtr<FJsonObject>& Field = ProjectEditorJson->GetObjectField(InProperty);
+		Field->SetStringField(FProjectEditorRecord::TimestampProperty, FDateTime::UtcNow().ToString());
+		return Field;
+	}
 
+	return nullptr;
 }
 
 TSharedPtr<FJsonObject> FProjectEditorRecord::MakeDefaultProperties()
@@ -91,29 +95,32 @@ TSharedPtr<FJsonObject> FProjectEditorRecord::MakeDefaultProperties()
 
 void FProjectEditorRecord::QueueUpdate(TUniqueFunction<void(FProjectEditorRecord&)>&& InUpdateFunction)
 {
-	TUniqueFunction<void()> Task = [Update = MoveTemp(InUpdateFunction)]()
+	if (GIsEditor && !IsRunningCommandlet())
 	{
-		FSystemWideCriticalSection SystemMutex(TEXT("ProjectEditorRecords"), FTimespan::FromMinutes(1));
-		FProjectEditorRecord AssociationFile = FProjectEditorRecord::Load();
-		Update(AssociationFile);
-		AssociationFile.Save();
-	};
+		TUniqueFunction<void()> Task = [Update = MoveTemp(InUpdateFunction)]()
+		{
+			FSystemWideCriticalSection SystemMutex(TEXT("ProjectEditorRecords"), FTimespan::FromMinutes(1));
+			FProjectEditorRecord AssociationFile = FProjectEditorRecord::Load();
+			Update(AssociationFile);
+			AssociationFile.Save();
+		};
 
-	if (FTaskGraphInterface::IsRunning())
-	{
-		if (AsyncUpdateTask)
+		if (FTaskGraphInterface::IsRunning())
 		{
-			AsyncUpdateTask = FFunctionGraphTask::CreateAndDispatchWhenReady(MoveTemp(Task), TStatId(), AsyncUpdateTask);
+			if (AsyncUpdateTask)
+			{
+				AsyncUpdateTask = FFunctionGraphTask::CreateAndDispatchWhenReady(MoveTemp(Task), TStatId(), AsyncUpdateTask);
+			}
+			else
+			{
+				AsyncUpdateTask = FFunctionGraphTask::CreateAndDispatchWhenReady(MoveTemp(Task), TStatId());
+			}
 		}
 		else
 		{
-			AsyncUpdateTask = FFunctionGraphTask::CreateAndDispatchWhenReady(MoveTemp(Task), TStatId());
+			Task();
 		}
 	}
-	else
-	{
-		Task();
-	}
 }
  
 void FProjectEditorRecord::TearDown()
diff --git a/Engine/Source/Runtime/Projects/Public/ProjectEditorRecords.h b/Engine/Source/Runtime/Projects/Public/ProjectEditorRecords.h
--- a/Engine/Source/Runtime/Projects/Public/ProjectEditorRecords.h
+++ b/Engine/Source/Runtime/Projects/Public/ProjectEditorRecords.h
@@ -65,7 +65,7 @@ struct FProjectEditorRecord
 	 * @param The name of the property to find or add
 	 * @return Shared pointer to the json object corresponding to the property input
 	 */
-	PROJECTS_API const TSharedPtr<FJsonObject>& FindOrAddProperty(const FString& InProperty);
+	PROJECTS_API const TSharedPtr<FJsonObject> FindOrAddProperty(const FString& InProperty);
 
 	/**
 	 * Creates a JsonObject containing the default entries 
