Converted .uplugin and .uproject loading to RapidJson for performance reasons

diff --git a/Engine/Source/Runtime/Json/Private/JsonUtils/JsonConversion.cpp b/Engine/Source/Runtime/Json/Private/JsonUtils/JsonConversion.cpp
new file mode 100644
--- /dev/null
+++ b/Engine/Source/Runtime/Json/Private/JsonUtils/JsonConversion.cpp
@@ -0,0 +1,229 @@
+// Copyright Epic Games, Inc. All Rights Reserved.
+
+#include "JsonUtils/JsonConversion.h"
+#include "Dom/JsonValue.h"
+#include "Dom/JsonObject.h"
+#include "Logging/StructuredLog.h"
+#include "Logging/LogMacros.h"
+
+
+DEFINE_LOG_CATEGORY_STATIC(LogJsonConversion, Display, All);
+
+namespace UE {
+namespace Json {
+
+namespace Private {
+
+	static TOptional<FValue> ConvertToRapidJsonValueInternal(const FJsonValue& SrcValue, FDocument::AllocatorType& Allocator)
+	{
+		switch (SrcValue.Type)
+		{
+		case EJson::None:
+			return {};
+
+		case EJson::Null:
+			return {FValue(rapidjson::kNullType)};
+
+		case EJson::String:
+			{
+				FString SrcString;
+
+				if (ensure(SrcValue.TryGetString(SrcString)))
+				{
+					FValue Result(rapidjson::kStringType);
+					Result.SetString(*SrcString, Allocator);
+
+					return {MoveTemp(Result)};
+				}				
+			}
+			return {};
+
+		case EJson::Number:
+			{
+				// all values in shared json discard their original types and are stored as doubles, this means that the reader functions have to do type conversions from doubles
+				// to mimic the legacy behaviour as long as a legacy API exists.
+
+				double SrcNumber = 0.0;
+				if (ensure(SrcValue.TryGetNumber(SrcNumber)))
+				{
+					FValue Result(rapidjson::kNumberType);
+					Result.SetDouble(SrcNumber);
+
+					return {MoveTemp(Result)};
+				}
+			}
+			return {};
+
+		case EJson::Boolean:
+			{
+				bool SrcBool = false;
+				if (ensure(SrcValue.TryGetBool(SrcBool)))
+				{
+					return {FValue(SrcBool)};
+				}				
+			}
+			return {};
+
+		case EJson::Array:
+			{
+				const TArray<TSharedPtr<FJsonValue>>* SrcArray = nullptr;
+
+				if (SrcValue.TryGetArray(SrcArray))
+				{
+					FValue DstArray(rapidjson::kArrayType);
+
+					for (const TSharedPtr<FJsonValue>& SrcArrayValue : *SrcArray)
+					{
+						TOptional<FValue> DstArrayValue = ConvertToRapidJsonValueInternal(*SrcArrayValue, Allocator);
+						if (DstArrayValue.IsSet())
+						{
+							DstArray.PushBack(MoveTemp(*DstArrayValue), Allocator);
+						}
+						else
+						{
+							// should have logged already
+							return {};
+						}
+					}
+
+					return {MoveTemp(DstArray)};
+				}				
+			}
+			return {};
+
+		case EJson::Object:
+			{
+				const TSharedPtr<FJsonObject>* SrcObject = nullptr;
+				if (SrcValue.TryGetObject(SrcObject))
+				{
+					FValue DstObject(rapidjson::kObjectType);
+					
+					for (const TPair<FString, TSharedPtr<FJsonValue>>& Member : (*SrcObject)->Values)
+					{
+						TOptional<FValue> DstMember = ConvertToRapidJsonValueInternal(*Member.Value, Allocator);
+						if (DstMember.IsSet())
+						{
+							FValue KeyType(rapidjson::kStringType);
+							KeyType.SetString(*Member.Key, Allocator);
+							DstObject.AddMember(MoveTemp(KeyType), MoveTemp(*DstMember), Allocator);
+						}
+						else
+						{
+							// should have been logged already
+							return {};
+						}
+					}
+
+					return {MoveTemp(DstObject)};
+				}
+			}
+			return {};
+		}
+
+		UE_LOG(LogJsonConversion, Error, TEXT("UE::Json::Private::ConvertToRapidJsonValueInternal Unable to convert unknown default Json type %d"), SrcValue.Type);
+		return {};
+	}
+}
+
+TSharedPtr<FJsonValue> ConvertRapidJsonToSharedJsonValue(const FValue& Value)
+{
+	switch (Value.GetType())
+	{
+	case rapidjson::kNullType:
+		return MakeShared<FJsonValueNull>();
+
+	case rapidjson::kFalseType:
+		return MakeShared<FJsonValueBoolean>(false);
+
+	case rapidjson::kTrueType:
+		return MakeShared<FJsonValueBoolean>(true);
+
+	case rapidjson::kObjectType:
+		{
+			TSharedPtr<FJsonObject> DstObject = MakeShared<FJsonObject>();
+
+			FConstObject SrcObject = Value.GetObject();
+			for (FValue::ConstMemberIterator It = SrcObject.MemberBegin(); It != SrcObject.MemberEnd(); ++It)
+			{
+				TSharedPtr<FJsonValue> ChildValue = ConvertRapidJsonToSharedJsonValue(It->value);
+				if (ChildValue.IsValid())
+				{
+					DstObject->SetField(It->name.GetString(), MoveTemp(ChildValue));
+				}
+			}
+
+			return MakeShared<FJsonValueObject>(MoveTemp(DstObject));
+		}
+
+	case rapidjson::kArrayType:
+		{
+			TArray<TSharedPtr<FJsonValue>> DstArray;
+
+			for (const FValue& SrcValue : Value.GetArray())
+			{
+				TSharedPtr<FJsonValue> DstValue = ConvertRapidJsonToSharedJsonValue(SrcValue);
+				if (DstValue.IsValid())
+				{
+					DstArray.Add(MoveTemp(DstValue));
+				}
+			}
+
+			return MakeShared<FJsonValueArray>(MoveTemp(DstArray));
+		}
+
+	case rapidjson::kStringType:
+		return MakeShared<FJsonValueString>(Value.GetString());
+
+	case rapidjson::kNumberType:
+		if (Value.IsDouble())
+		{
+			return MakeShared<FJsonValueNumber>(Value.GetDouble());
+		}
+		else if (Value.IsInt64())
+		{
+			return MakeShared<FJsonValueNumber>(static_cast<double>(Value.GetInt64()));
+		}
+		else if (Value.IsUint64())
+		{
+			return MakeShared<FJsonValueNumber>(static_cast<double>(Value.GetUint64()));
+		}
+	}
+
+	UE_LOG(LogJsonConversion, Warning, TEXT("UE::Json::ConvertRapidJsonToSharedJsonValue Unhandled value type %d"), int32(Value.GetType()));
+	return {};
+}
+
+TSharedPtr<FJsonObject> ConvertRapidJsonToSharedJsonObject(FConstObject Object)
+{
+	TSharedPtr<FJsonValue> ResultValue = ConvertRapidJsonToSharedJsonValue(Object);
+	if (!ResultValue.IsValid() || ResultValue->Type != EJson::Object)
+	{
+		UE_LOG(LogJsonConversion, Error, TEXT("UE::Json::ConvertRapidJsonToSharedJsonObject unable to convert JSON object"));
+		return {};
+	}
+
+	return ResultValue->AsObject();
+}
+
+TOptional<FDocument> ConvertSharedJsonToRapidJsonDocument(const FJsonObject& SrcObject)
+{
+	FDocument DstDocument;
+	DstDocument.SetObject();
+	FDocument::AllocatorType& Allocator = DstDocument.GetAllocator();
+
+	for (const TPair<FString, TSharedPtr<FJsonValue>>& SrcField : SrcObject.Values)
+	{
+		TOptional<FValue> DstMember = Private::ConvertToRapidJsonValueInternal(*SrcField.Value, Allocator);
+		if (DstMember.IsSet())
+		{
+			FValue KeyType(rapidjson::kStringType);
+			KeyType.SetString(*SrcField.Key, Allocator);
+
+			DstDocument.AddMember(MoveTemp(KeyType), MoveTemp(*DstMember), Allocator);
+		}
+	}
+
+	return MoveTemp(DstDocument);
+}
+
+}} // namespace UE::Json
diff --git a/Engine/Source/Runtime/Json/Private/JsonUtils/RapidJsonUtils.cpp b/Engine/Source/Runtime/Json/Private/JsonUtils/RapidJsonUtils.cpp
new file mode 100644
--- /dev/null
+++ b/Engine/Source/Runtime/Json/Private/JsonUtils/RapidJsonUtils.cpp
@@ -0,0 +1,304 @@
+// Copyright Epic Games, Inc. All Rights Reserved.
+
+#include "JsonUtils/RapidJsonUtils.h"
+#include "Dom/JsonValue.h"
+
+namespace UE {
+namespace Json {
+
+const constexpr uint32 DefaultParseFlags = rapidjson::ParseFlag::kParseTrailingCommasFlag;
+
+TOptional<bool> GetBoolField(FConstObject Object, const TCHAR* FieldName)
+{
+	FValue::ConstMemberIterator It = Object.FindMember(FieldName);
+	if (It == Object.MemberEnd() || !It->value.IsBool())
+	{
+		return {};
+	}
+
+	return It->value.GetBool();
+}
+
+TOptional<int32> GetInt32Field(FConstObject Object, const TCHAR* FieldName)
+{
+	FValue::ConstMemberIterator It = Object.FindMember(FieldName);
+	if (It == Object.MemberEnd() || !It->value.IsInt())
+	{
+		return {};
+	}
+
+	return It->value.GetInt();
+}
+
+TOptional<uint32> GetUint32Field(FConstObject Object, const TCHAR* FieldName)
+{
+	FValue::ConstMemberIterator It = Object.FindMember(FieldName);
+	if (It == Object.MemberEnd() || !It->value.IsUint())
+	{
+		return {};
+	}
+
+	return It->value.GetUint();
+}
+
+TOptional<int64> GetInt64Field(FConstObject Object, const TCHAR* FieldName)
+{
+	FValue::ConstMemberIterator It = Object.FindMember(FieldName);
+	if (It == Object.MemberEnd() || !It->value.IsInt64())
+	{
+		return {};
+	}
+
+	return It->value.GetInt64();
+}
+
+TOptional<uint64> GetUint64Field(FConstObject Object, const TCHAR* FieldName)
+{
+	FValue::ConstMemberIterator It = Object.FindMember(FieldName);
+	if (It == Object.MemberEnd() || !It->value.IsUint64())
+	{
+		return {};
+	}
+
+	return It->value.GetUint64();
+}
+
+
+TOptional<double> GetDoubleField(FConstObject Object, const TCHAR* FieldName)
+{
+	FValue::ConstMemberIterator It = Object.FindMember(FieldName);
+	if (It == Object.MemberEnd() || !It->value.IsDouble())
+	{
+		return {};
+	}
+
+	return It->value.GetDouble();
+}
+
+TOptional<FStringView> GetStringField(FConstObject Object, const TCHAR* FieldName)
+{
+	FValue::ConstMemberIterator It = Object.FindMember(FieldName);
+	if (It == Object.MemberEnd() || !It->value.IsString())
+	{
+		return {};
+	}
+
+	return FStringView(It->value.GetString(), It->value.GetStringLength());
+}
+
+bool HasField(FConstObject Object, const TCHAR* FieldName)
+{
+	FValue::ConstMemberIterator It = Object.FindMember(FieldName);
+	return It != Object.MemberEnd();
+}
+
+bool HasNullField(FConstObject Object, const TCHAR* FieldName)
+{
+	FValue::ConstMemberIterator It = Object.FindMember(FieldName);
+	return It != Object.MemberEnd() && It->value.GetType() == rapidjson::kNullType;
+}
+
+TOptional<FConstObject> GetObjectField(FConstObject Object, const TCHAR* FieldName)
+{
+	FValue::ConstMemberIterator It = Object.FindMember(FieldName);
+	if (It == Object.MemberEnd() || !It->value.IsObject())
+	{
+		return {};
+	}
+
+	return {It->value.GetObject()};
+}
+
+TOptional<FConstObject> GetRootObject(const FDocument& Document)
+{
+	if (Document.IsObject())
+	{
+		return {Document.GetObject()};
+	}
+
+	return {};
+}
+
+TOptional<FConstArray> GetArrayField(FConstObject Object, const TCHAR* FieldName)
+{
+	FValue::ConstMemberIterator It = Object.FindMember(FieldName);
+	if (It == Object.MemberEnd() || !It->value.IsArray())
+	{
+		return {};
+	}
+
+	return {It->value.GetArray()};
+}
+
+int32 FindLineNumber(const FStringView JsonText, uint32 Offset)
+{
+	// 1-index line numbers
+	int32 NCount = 1;
+	int32 RCount = 1;
+
+	FStringView BeforeError = JsonText.Left(Offset);
+
+	// count number of newlines up util the offset
+	for (TCHAR Current : BeforeError)
+	{
+		if (Current == TCHAR('\n'))
+		{
+			NCount++;
+		}
+		else if (Current == TCHAR('\r'))
+		{
+			RCount++;
+		}
+	}
+
+	return FMath::Max(NCount, RCount);
+}
+
+static int32 FindLineStartOffset(const FStringView JsonText, uint32 LineNum)
+{
+	if (LineNum <= 1)
+	{
+		return 0; // start at the beginning of the file
+	}
+
+	int32 CurrentNCount = 0;
+	int32 CurrentRCount = 0;
+	int32 CurrentOffset = 0;
+
+	for (TCHAR Current : JsonText)
+	{
+		++CurrentOffset;
+
+		// track \r and \n separately, we care about the largest one
+		if (Current == TCHAR('\n'))
+		{
+			CurrentNCount++;
+			if (CurrentNCount == LineNum)
+			{
+				return CurrentOffset;
+			}
+		}
+		else if (Current == TCHAR('\r'))
+		{
+			CurrentRCount++;
+			if (CurrentRCount == LineNum)
+			{
+				return CurrentOffset;
+			}
+		}
+	}
+
+	return 0;
+}
+
+constexpr size_t InternalError_NotNullTerminated = 1;
+
+FString FParseError::CreateMessage(const FStringView JsonText) const
+{
+	if (ErrorCode == rapidjson::kParseErrorNone)
+	{
+		// in thie case we use the offset as an extended error code 
+		if (Offset == InternalError_NotNullTerminated)
+		{
+			return TEXT("JsonText must be null terminated");
+		}
+		else
+		{
+			return TEXT("Invalid parse error");
+		}
+	}
+	else
+	{
+		int32 LineNum = FindLineNumber(JsonText, Offset);
+
+		int32 SnippetStart = FindLineStartOffset(JsonText, LineNum-2);
+		int32 SnippetEnd = FindLineStartOffset(JsonText, LineNum+1);
+
+		FString Snippet(JsonText.Mid(SnippetStart, SnippetEnd-SnippetStart).TrimStartAndEnd());
+
+		return FString::Printf(
+			TEXT("%s, Line %d:\n%s"),
+			GetParseError_En(ErrorCode),
+			LineNum,
+			*Snippet
+		);
+	}
+}
+
+#if PLATFORM_LITTLE_ENDIAN
+	using FDocumentEncoding = rapidjson::UTF16LE<TCHAR>;
+#else
+	using FDocumentEncoding = rapidjson::UTF16BE<TCHAR>;
+#endif
+
+TValueOrError<FDocument, FParseError> Parse(const FStringView JsonText)
+{
+	FDocument Result;
+	Result.Parse<DefaultParseFlags, FDocumentEncoding>(JsonText.GetData(), JsonText.Len());
+
+	if (Result.HasParseError())
+	{
+		return MakeError(FParseError{
+			.ErrorCode = Result.GetParseError(),
+			.Offset = Result.GetErrorOffset() / sizeof(TCHAR)
+		});
+	}
+
+	return MakeValue(MoveTemp(Result));
+}
+
+TValueOrError<FDocument, FParseError> ParseInPlace(TArrayView<TCHAR> JsonText)
+{
+	if (!JsonText.IsEmpty() && JsonText[JsonText.Num()-1] != 0)
+	{
+		return MakeError(FParseError{
+			.Offset = InternalError_NotNullTerminated
+		});
+	}
+
+	FDocument Result;
+
+	rapidjson::GenericInsituStringStream<FDocumentEncoding> Stream(JsonText.GetData());
+	Result.ParseStream<DefaultParseFlags | rapidjson::kParseInsituFlag>(Stream);
+
+	if (Result.HasParseError())
+	{
+		return MakeError(FParseError{
+			.ErrorCode = Result.GetParseError(),
+			.Offset = Result.GetErrorOffset() / sizeof(TCHAR)
+		});
+	}
+
+	return MakeValue(MoveTemp(Result));
+}
+
+const TCHAR* GetValueTypeName(const FValue& Value)
+{
+	switch (Value.GetType())
+	{
+	case rapidjson::kNullType:
+		return TEXT("Null");
+
+	case rapidjson::kTrueType:
+	case rapidjson::kFalseType:
+		return TEXT("Bool");
+
+	case rapidjson::kNumberType:
+		return TEXT("Number");
+
+	case rapidjson::kArrayType:
+		return TEXT("Array");
+
+	case rapidjson::kObjectType:
+		return TEXT("Object");
+
+	case rapidjson::kStringType:
+		return TEXT("String");
+
+	default:
+		return TEXT("Unknown");
+	}
+}
+
+} // namespace Json
+} // namespace UE
diff --git a/Engine/Source/Runtime/Json/Private/Tests/JsonConversionTests.cpp b/Engine/Source/Runtime/Json/Private/Tests/JsonConversionTests.cpp
new file mode 100644
--- /dev/null
+++ b/Engine/Source/Runtime/Json/Private/Tests/JsonConversionTests.cpp
@@ -0,0 +1,175 @@
+// Copyright Epic Games, Inc. All Rights Reserved.
+
+#if WITH_TESTS
+
+#include "CoreMinimal.h"
+#include "JsonUtils/JsonConversion.h"
+#include "Misc/AutomationTest.h"
+#include "Serialization/JsonSerializer.h"
+#include "Tests/TestHarnessAdapter.h"
+
+#if WITH_LOW_LEVEL_TESTS
+#include "TestCommon/Expectations.h"
+#endif
+
+TEST_CASE_NAMED(FJsonConversionAutomationTest, "System::Engine::FileSystem::JSON::Conversion", "[ApplicationContextMask][SmokeFilter]")
+{
+	const TCHAR* JsonText = TEXT(R"json(
+{
+	"int": 1234,
+	"round_down": 1.1,
+	"round_up" : 1.9,
+	"string": "hello world",
+	"on": true,
+	"off": false,
+	"object": {
+		"value": 42
+	},
+	"array": [
+		"aaa",
+		"bbb",
+		"ccc",
+		"ddd"
+	],	
+	"null": null
+}
+
+)json");
+
+	// Checks to make a RapidJson -> default Json conversion is identical
+	auto CheckEqual = [&](const TSharedPtr<FJsonObject>& JsonObjectA, const UE::Json::FDocument& RapidJsonDocument)
+	{
+		TOptional<UE::Json::FConstObject> RootObject = UE::Json::GetRootObject(RapidJsonDocument);
+		REQUIRE(RootObject.IsSet());
+
+		TSharedPtr<FJsonObject> JsonObjectB = UE::Json::ConvertRapidJsonToSharedJsonObject(*RootObject);
+		REQUIRE(JsonObjectB);
+
+		TSharedPtr<FJsonValue> IntValueA = JsonObjectA->TryGetField(TEXT("int"));
+		TSharedPtr<FJsonValue> IntValueB = JsonObjectB->TryGetField(TEXT("int"));
+		REQUIRE(IntValueA);
+		REQUIRE(IntValueB);
+		REQUIRE(IntValueA->Type == EJson::Number);
+		REQUIRE(IntValueB->Type == EJson::Number);
+		REQUIRE(IntValueA->AsNumber() == IntValueB->AsNumber());
+
+		TSharedPtr<FJsonValue> RoundDownValueA = JsonObjectA->TryGetField(TEXT("round_down"));
+		TSharedPtr<FJsonValue> RoundDownValueB = JsonObjectB->TryGetField(TEXT("round_down"));
+		REQUIRE(RoundDownValueA);
+		REQUIRE(RoundDownValueB);
+		REQUIRE(RoundDownValueA->Type == EJson::Number);
+		REQUIRE(RoundDownValueB->Type == EJson::Number);
+		REQUIRE(RoundDownValueA->AsNumber() == RoundDownValueB->AsNumber());
+		// default JSON does some rounding when converting, ensure that is kept during the conversion
+		int32 RoundDownA = 0;
+		int32 RoundDownB = 0;
+		REQUIRE(RoundDownValueA->TryGetNumber(RoundDownA));
+		REQUIRE(RoundDownValueB->TryGetNumber(RoundDownB));
+		REQUIRE(RoundDownA == RoundDownB);
+
+		TSharedPtr<FJsonValue> RoundUpValueA = JsonObjectA->TryGetField(TEXT("round_down"));
+		TSharedPtr<FJsonValue> RoundUpValueB = JsonObjectB->TryGetField(TEXT("round_down"));
+		REQUIRE(RoundUpValueA);
+		REQUIRE(RoundUpValueB);
+		REQUIRE(RoundUpValueA->Type == EJson::Number);
+		REQUIRE(RoundUpValueB->Type == EJson::Number);
+		REQUIRE(RoundUpValueA->AsNumber() == RoundUpValueB->AsNumber());
+		// default JSON does some rounding when converting from float to integer, ensure that is kept during the conversion
+		int32 RoundUpA = 0;
+		int32 RoundUpB = 0;
+		REQUIRE(RoundUpValueA->TryGetNumber(RoundUpA));
+		REQUIRE(RoundUpValueB->TryGetNumber(RoundUpB));
+		REQUIRE(RoundUpA == RoundUpB);
+
+		TSharedPtr<FJsonValue> StringValueA = JsonObjectA->TryGetField(TEXT("string"));
+		TSharedPtr<FJsonValue> StringValueB = JsonObjectB->TryGetField(TEXT("string"));
+		REQUIRE(StringValueA);
+		REQUIRE(StringValueB);
+		REQUIRE(StringValueA->Type == EJson::String);
+		REQUIRE(StringValueB->Type == EJson::String);
+		REQUIRE(StringValueA->AsString() == StringValueB->AsString());
+
+		TSharedPtr<FJsonValue> OnValueA = JsonObjectA->TryGetField(TEXT("on"));
+		TSharedPtr<FJsonValue> OnValueB = JsonObjectB->TryGetField(TEXT("on"));
+		REQUIRE(OnValueA);
+		REQUIRE(OnValueB);
+		REQUIRE(OnValueA->Type == EJson::Boolean);
+		REQUIRE(OnValueB->Type == EJson::Boolean);
+		REQUIRE(OnValueA->AsBool() == OnValueB->AsBool());
+
+		TSharedPtr<FJsonValue> OffValueA = JsonObjectA->TryGetField(TEXT("off"));
+		TSharedPtr<FJsonValue> OffValueB = JsonObjectB->TryGetField(TEXT("off"));
+		REQUIRE(OffValueA);
+		REQUIRE(OffValueB);
+		REQUIRE(OffValueA->Type == EJson::Boolean);
+		REQUIRE(OffValueB->Type == EJson::Boolean);
+		REQUIRE(OffValueA->AsBool() == OffValueB->AsBool());
+
+		TSharedPtr<FJsonValue> ObjectValueA = JsonObjectA->TryGetField(TEXT("object"));
+		TSharedPtr<FJsonValue> ObjectValueB = JsonObjectB->TryGetField(TEXT("object"));
+		REQUIRE(ObjectValueA);
+		REQUIRE(ObjectValueB);
+		REQUIRE(ObjectValueA->Type == EJson::Object);
+		REQUIRE(ObjectValueB->Type == EJson::Object);
+
+		TSharedPtr<FJsonObject> ObjectA = ObjectValueA->AsObject();
+		TSharedPtr<FJsonObject> ObjectB = ObjectValueB->AsObject();
+		REQUIRE(ObjectA);
+		REQUIRE(ObjectB);
+
+		TSharedPtr<FJsonValue> ObjectFieldA = ObjectA->TryGetField(TEXT("value"));		
+		TSharedPtr<FJsonValue> ObjectFieldB = ObjectB->TryGetField(TEXT("value"));
+		REQUIRE(ObjectFieldA);
+		REQUIRE(ObjectFieldB);
+		REQUIRE(ObjectFieldA->Type == EJson::Number);
+		REQUIRE(ObjectFieldB->Type == EJson::Number);
+		REQUIRE(ObjectFieldA->AsNumber() == ObjectFieldB->AsNumber());
+
+		TSharedPtr<FJsonValue> ArrayValueA = JsonObjectA->TryGetField(TEXT("array"));
+		TSharedPtr<FJsonValue> ArrayValueB = JsonObjectB->TryGetField(TEXT("array"));
+		REQUIRE(ArrayValueA);
+		REQUIRE(ArrayValueB);
+		REQUIRE(ArrayValueA->Type == EJson::Array);
+		REQUIRE(ArrayValueB->Type == EJson::Array);
+
+		TArray<TSharedPtr<FJsonValue>> ArrayA = ArrayValueA->AsArray();
+		TArray<TSharedPtr<FJsonValue>> ArrayB = ArrayValueB->AsArray();
+
+		REQUIRE(!ArrayA.IsEmpty());
+		REQUIRE(!ArrayB.IsEmpty());
+
+		REQUIRE(ArrayA.Num() == ArrayB.Num());
+
+		for (int32 Idx = 0; Idx < ArrayA.Num(); ++Idx)
+		{
+			const TSharedPtr<FJsonValue>& ItemA = ArrayA[Idx];
+			const TSharedPtr<FJsonValue>& ItemB = ArrayB[Idx];
+
+			REQUIRE(ItemA);
+			REQUIRE(ItemB);
+			REQUIRE(ItemA->Type == EJson::String);
+			REQUIRE(ItemB->Type == EJson::String);
+			REQUIRE(ItemA->AsString() == ItemB->AsString());
+		}
+	};
+
+
+	TSharedRef< TJsonReader<> > Reader = TJsonReaderFactory<>::Create(JsonText);
+	TSharedPtr< FJsonObject > DefaultLoadedJson;
+	REQUIRE(FJsonSerializer::Deserialize(Reader, DefaultLoadedJson));
+
+	// convert to RapidJson
+	TOptional<UE::Json::FDocument> ConvertedDocument = UE::Json::ConvertSharedJsonToRapidJsonDocument(*DefaultLoadedJson);
+	REQUIRE(ConvertedDocument.IsSet());
+	
+	CheckEqual(DefaultLoadedJson, *ConvertedDocument);
+
+	TValueOrError<UE::Json::FDocument, UE::Json::FParseError> LoadedDocument = UE::Json::Parse(JsonText);
+
+	FString ErrorMessage = LoadedDocument.HasError() ? LoadedDocument.GetError().CreateMessage(JsonText) : FString();
+	REQUIRE_MESSAGE(ErrorMessage, LoadedDocument.HasValue());
+
+	CheckEqual(DefaultLoadedJson, LoadedDocument.GetValue());
+}
+
+#endif // WITH_TESTS
diff --git a/Engine/Source/Runtime/Json/Private/Tests/RapidJsonUtilsTests.cpp b/Engine/Source/Runtime/Json/Private/Tests/RapidJsonUtilsTests.cpp
new file mode 100644
--- /dev/null
+++ b/Engine/Source/Runtime/Json/Private/Tests/RapidJsonUtilsTests.cpp
@@ -0,0 +1,185 @@
+// Copyright Epic Games, Inc. All Rights Reserved.
+
+#if WITH_TESTS
+
+#include "CoreMinimal.h"
+#include "JsonUtils/JsonConversion.h"
+#include "Misc/AutomationTest.h"
+#include "Serialization/JsonSerializer.h"
+#include "Tests/TestHarnessAdapter.h"
+
+#if WITH_LOW_LEVEL_TESTS
+#include "TestCommon/Expectations.h"
+#endif
+
+/** Ensure RapidJsonUtils parses with the same quirks as the default Json loader, e.g. converting between types correctly **/
+TEST_CASE_NAMED(FRapidJsonUtilsTests, "System::Engine::FileSystem::RapidJSONUtils::Basic", "[ApplicationContextMask][SmokeFilter]")
+{
+	const TCHAR* JsonText = TEXT(R"json(
+{
+	"on": true,
+	"off": false,
+	"zero": 0,
+	"min_int32": -2147483648,
+	"max_int32":  2147483647,
+	"max_uint32": 4294967295,
+	"min_int64": -9223372036854775808,
+	"max_int64":  9223372036854775807,
+	"max_uint64": 18446744073709551615,
+	"double": 1.1,
+	"big_double": 18446744073709551615.18446744073709551615,
+	"string": "hello world",
+	"object": {
+		"value": 42
+	},
+	"array": [
+		"aaa",
+		"bbb",
+		"ccc",
+		"ddd"
+	],	
+	"null": null
+}
+
+)json");
+
+	// check to make sure all values are parsed the same across RapidJsonUtils an default Json
+	// convert to RapidJson
+
+	using namespace UE::Json;
+
+	TValueOrError<FDocument, FParseError> RapidJsonDocument = Parse(JsonText);
+	FString ParseErrorMessage = RapidJsonDocument.HasError() ? RapidJsonDocument.GetError().CreateMessage(JsonText) : FString();
+	REQUIRE_MESSAGE(ParseErrorMessage, RapidJsonDocument.HasValue());
+
+	TOptional<FConstObject> Root = GetRootObject(RapidJsonDocument.GetValue());
+	REQUIRE(Root.IsSet());
+
+	TOptional<bool> OnField = GetBoolField(*Root, TEXT("on"));
+	REQUIRE(OnField.IsSet());
+	REQUIRE(*OnField == true);
+
+	TOptional<bool> OffField = GetBoolField(*Root, TEXT("off"));
+	REQUIRE(OffField.IsSet());
+	REQUIRE(*OffField == false);
+
+	TOptional<int32> ZeroField = GetInt32Field(*Root, TEXT("zero"));
+	REQUIRE(ZeroField.IsSet());
+	REQUIRE(*ZeroField == 0);
+
+	TOptional<int32> MinInt32Field = GetInt32Field(*Root, TEXT("min_int32"));
+	REQUIRE(MinInt32Field.IsSet());
+	REQUIRE(*MinInt32Field == MIN_int32);
+
+	TOptional<int32> MaxInt32Field = GetInt32Field(*Root, TEXT("max_int32"));
+	REQUIRE(MaxInt32Field.IsSet());
+	REQUIRE(*MaxInt32Field == MAX_int32);
+
+	TOptional<uint32> MaxUint32Field = GetUint32Field(*Root, TEXT("max_uint32"));
+	REQUIRE(MaxUint32Field.IsSet());
+	REQUIRE(*MaxUint32Field == MAX_uint32);
+
+	TOptional<int64> MinInt64Field = GetInt64Field(*Root, TEXT("min_int64"));
+	REQUIRE(MinInt64Field.IsSet());
+	REQUIRE(*MinInt64Field == MIN_int64);
+
+	TOptional<int64> MaxInt64Field = GetInt64Field(*Root, TEXT("max_int64"));
+	REQUIRE(MaxInt64Field.IsSet());
+	REQUIRE(*MaxInt64Field == MAX_int64);
+
+	TOptional<uint64> MaxUint64Field = GetUint64Field(*Root, TEXT("max_uint64"));
+	REQUIRE(MaxUint64Field.IsSet());
+	REQUIRE(*MaxUint64Field == MAX_uint64);
+
+	TOptional<double> DoubleField = GetDoubleField(*Root, TEXT("double"));
+	REQUIRE(DoubleField.IsSet());
+	REQUIRE(*DoubleField == 1.1);
+
+	TOptional<double> BigDoubleField = GetDoubleField(*Root, TEXT("big_double"));
+	REQUIRE(BigDoubleField.IsSet());
+	REQUIRE(*BigDoubleField == 18446744073709551615.18446744073709551615);
+
+	TOptional<FStringView> StringField = GetStringField(*Root, TEXT("string"));
+	REQUIRE(StringField.IsSet());
+	REQUIRE(StringField->Equals(TEXT("hello world")));
+
+	TOptional<FConstObject> ObjectField = GetObjectField(*Root, TEXT("object"));
+	REQUIRE(StringField.IsSet());
+	TOptional<int32> ObjectValueField = GetInt32Field(*ObjectField, TEXT("value"));
+	REQUIRE(ObjectValueField.IsSet());
+	REQUIRE(*ObjectValueField == 42);
+
+	TOptional<FConstArray> ArrayField = GetArrayField(*Root, TEXT("array"));
+	REQUIRE(ArrayField.IsSet());
+	REQUIRE(ArrayField->Size() == 4);
+
+	REQUIRE(HasNullField(*Root, TEXT("null")));
+
+	REQUIRE(!GetBoolField(*Root, TEXT("missing")));
+	REQUIRE(!GetInt32Field(*Root, TEXT("missing")));
+	REQUIRE(!GetUint32Field(*Root, TEXT("missing")));
+	REQUIRE(!GetInt64Field(*Root, TEXT("missing")));
+	REQUIRE(!GetUint64Field(*Root, TEXT("missing")));
+	REQUIRE(!GetDoubleField(*Root, TEXT("missing")));
+	REQUIRE(!GetObjectField(*Root, TEXT("missing")));
+	REQUIRE(!GetStringField(*Root, TEXT("missing")));
+	REQUIRE(!GetArrayField(*Root, TEXT("missing")));
+	REQUIRE(!HasNullField(*Root, TEXT("missing")));
+
+}
+
+/** Ensure RapidJsonUtils parses with the same quirks as the default Json loader, e.g. converting between types correctly **/
+TEST_CASE_NAMED(FRapidJsonUtilsErrorHandlingTests, "System::Engine::FileSystem::RapidJSONUtils::ErrorHandling", "[ApplicationContextMask][SmokeFilter]")
+{
+	// check to make sure all values are parsed the same across RapidJsonUtils an default Json
+	// convert to RapidJson
+
+	using namespace UE::Json;
+
+	// check newlines work correctly across all versions
+	const TCHAR* JsonUnixText = TEXT("{\n\"ok\": 42,\n\"bad\": xyz\n}");
+	const TCHAR* JsonMacText = TEXT("{\r\"ok\": 42,\r\"bad\": xyz\r}");
+	const TCHAR* JsonWindowsText = TEXT("{\r\n\"ok\": 42,\r\n\"bad\": xyz\r\n}");
+
+	TValueOrError<FDocument, FParseError> UnixDoc = Parse(JsonUnixText);
+	REQUIRE(UnixDoc.HasError());
+	FString UnixError = UnixDoc.GetError().CreateMessage(JsonUnixText);
+	REQUIRE_MESSAGE(UnixError, UnixError.Contains(TEXT("Line 3")));
+
+	TValueOrError<FDocument, FParseError> MacDoc = Parse(JsonMacText);
+	REQUIRE(MacDoc.HasError());
+	FString MacError = MacDoc.GetError().CreateMessage(JsonMacText);
+	REQUIRE_MESSAGE(MacError, MacError.Contains(TEXT("Line 3")));
+
+	TValueOrError<FDocument, FParseError> WindowsDoc = Parse(JsonWindowsText);
+	REQUIRE(WindowsDoc.HasError());
+	FString WindowsError = WindowsDoc.GetError().CreateMessage(JsonWindowsText);
+	REQUIRE_MESSAGE(WindowsError, WindowsError.Contains(TEXT("Line 3")));
+
+	const TCHAR* LongJsonText = TEXT(R"json({
+	"array": [
+		0,
+		1,
+		2,
+		3,
+		4,
+		5,
+		6,
+		7,
+		bad_format, 
+		9,
+		10,
+		11
+	]
+})json");
+
+	TValueOrError<FDocument, FParseError> LongDoc = Parse(LongJsonText);
+	REQUIRE(LongDoc.HasError());
+	FString LongError = LongDoc.GetError().CreateMessage(LongJsonText);
+
+	// make sure the snippet contains our issue
+	REQUIRE_MESSAGE(LongError, LongError.Contains(TEXT("bad_format")));
+
+}
+
+#endif // WITH_TESTS
diff --git a/Engine/Source/Runtime/Json/Public/JsonUtils/JsonConversion.h b/Engine/Source/Runtime/Json/Public/JsonUtils/JsonConversion.h
new file mode 100644
--- /dev/null
+++ b/Engine/Source/Runtime/Json/Public/JsonUtils/JsonConversion.h
@@ -0,0 +1,30 @@
+// Copyright Epic Games, Inc. All Rights Reserved.
+
+#pragma once
+
+#include "Templates/SharedPointer.h"
+#include "JsonUtils/RapidJsonUtils.h"
+
+class FJsonValue;
+class FJsonObject;
+
+namespace UE {
+namespace Json {
+
+/**
+ * Converts from a RapidJSON value to a shared json value
+ */
+JSON_API TSharedPtr<FJsonValue> ConvertRapidJsonToSharedJsonValue(const FValue& Value);
+
+/**
+ * Converts from a RapidJSON object to a shared json object
+ */
+JSON_API TSharedPtr<FJsonObject> ConvertRapidJsonToSharedJsonObject(FConstObject Object);
+
+/**
+ * Converts from shared JSON object to a RapidJSON Document.
+ */
+JSON_API TOptional<FDocument> ConvertSharedJsonToRapidJsonDocument(const FJsonObject& SrcObject);
+
+}
+}
diff --git a/Engine/Source/Runtime/Json/Public/JsonUtils/RapidJsonUtils.h b/Engine/Source/Runtime/Json/Public/JsonUtils/RapidJsonUtils.h
new file mode 100644
--- /dev/null
+++ b/Engine/Source/Runtime/Json/Public/JsonUtils/RapidJsonUtils.h
@@ -0,0 +1,143 @@
+// Copyright Epic Games, Inc. All Rights Reserved.
+
+#pragma once
+
+#include "RapidJsonIncludes.h"
+#include "Misc/Optional.h"
+#include "Containers/StringView.h"
+#include "Templates/ValueOrError.h"
+
+/**
+ * Helper functions for reading from RapidJSON into unreal types
+ */
+namespace UE {
+namespace Json {
+
+	// types that use TCHAR
+	using FDocument = rapidjson::GenericDocument<rapidjson::UTF16<TCHAR>, rapidjson::CrtAllocator>;
+	using FValue = FDocument::ValueType;
+	using FConstObject = FValue::ConstObject;
+	using FConstArray = FValue::ConstArray;
+	using FMember = FValue::Member;
+
+	/**
+	 * Attempts to read a boolean with a specified name.
+	 * 
+	 * Returns the value, if found and matches the type explicitly
+	 */
+	JSON_API TOptional<bool> GetBoolField(FConstObject Object, const TCHAR* FieldName);
+
+	/**
+	 * Attempts to read an int32 with a specified name.
+	 * 
+	 * Returns the value, if found and matches the type explicitly
+	 */
+	JSON_API TOptional<int32> GetInt32Field(FConstObject Object, const TCHAR* FieldName);
+
+	/**
+	 * Attempts to read an uint32 with a specified name.
+	 * 
+	 * Returns the value, if found and matches the type explicitly
+	 */
+	JSON_API TOptional<uint32> GetUint32Field(FConstObject Object, const TCHAR* FieldName);
+
+	/**
+	 * Attempts to read an int64 with a specified name.
+	 * 
+	 * Returns the value, if found and matches the type explicitly
+	 */
+	JSON_API TOptional<int64> GetInt64Field(FConstObject Object, const TCHAR* FieldName);
+
+	/**
+	 * Attempts to read an uint64 with a specified name.
+	 * 
+	 * Returns the value, if found and matches the type explicitly
+	 */
+	JSON_API TOptional<uint64> GetUint64Field(FConstObject Object, const TCHAR* FieldName);
+
+	/**
+	 * Attempts to read an double with a specified name.
+	 * 
+	 * Returns the value, if found and matches the type explicitly
+	 */
+	JSON_API TOptional<double> GetDoubleField(FConstObject Object, const TCHAR* FieldName);
+
+	/**
+	 * Attempts to read a string with a specified name
+	 * 
+	 * Returns the value, if found and matches the type explicitly
+	 */
+	JSON_API TOptional<FStringView> GetStringField(FConstObject Object, const TCHAR* FieldName);
+
+	/**
+	 * Attempts to parse an enum from a string with a specified name
+	 * 
+	 * Returns the value, if found and matches the type explicitly
+	 */
+	template<typename TEnum>
+	inline TOptional<TEnum> GetEnumField(FConstObject Object, const TCHAR* FieldName)
+	{
+		TEnum Value;
+		TOptional<FStringView> StringView = GetStringField(Object, FieldName);
+
+		if (StringView.IsSet() && LexTryParseString(Value, *StringView))
+		{
+			return Value;
+		}
+
+		return {};
+	}
+
+	/**
+	 * Attempts to find an object with the specified namennn
+	 * 
+	 * Returns a valid object ref if the field was found and it matches an object.
+	 */
+	JSON_API TOptional<FConstObject> GetObjectField(FConstObject Object, const TCHAR* FieldName);
+
+	/**
+	 * Attempts to find an array with the specified name
+	 * 
+	 * Returns a valid array ref if the field was found and it matches an array
+	 */
+	JSON_API TOptional<FConstArray> GetArrayField(FConstObject Object, const TCHAR* FieldName);
+
+	/**
+	 * Returns true if a field exists that is a null type
+	 */
+	JSON_API bool HasNullField(FConstObject Object, const TCHAR* FieldName);
+
+	/**
+	 * Returns the root object, if that was the root type for the document
+	 */
+	JSON_API TOptional<FConstObject> GetRootObject(const FDocument& Document);
+
+	struct FParseError
+	{
+		rapidjson::ParseErrorCode ErrorCode = rapidjson::kParseErrorNone;
+		size_t Offset = 0;
+
+		/**
+		 * Create a user-readable error message. Takes the source text paramter to find the line number from the offset
+		 */
+		JSON_API FString CreateMessage(const FStringView JsonText) const;
+	};
+
+	/**
+	 * Parse RapidJSON with default flags
+	 */
+	JSON_API TValueOrError<FDocument, FParseError> Parse(const FStringView JsonText);
+
+	/**
+	 * Parse RapidJSON with default flags, into destructible memory, to reduce allocation counts.
+	 * 
+	 * JsonText must be a zero terminated string per the rapidjson documentation
+	 */
+	JSON_API TValueOrError<FDocument, FParseError> ParseInPlace(TArrayView<TCHAR> JsonText);
+
+	/**
+	 * Return a string representation of type name, for debug logging purposes
+	 */
+	JSON_API const TCHAR* GetValueTypeName(const FValue& Value);
+}
+}
diff --git a/Engine/Source/Runtime/Json/Public/RapidJsonIncludes.h b/Engine/Source/Runtime/Json/Public/RapidJsonIncludes.h
--- a/Engine/Source/Runtime/Json/Public/RapidJsonIncludes.h
+++ b/Engine/Source/Runtime/Json/Public/RapidJsonIncludes.h
@@ -3,6 +3,7 @@
 #pragma once
 
 #include "HAL/Platform.h"
+#include "HAL/UnrealMemory.h"
 #include "Misc/AssertionMacros.h"
 #include "Misc/CoreMiscDefines.h"
 
@@ -25,6 +26,30 @@ THIRD_PARTY_INCLUDES_START
 #define RAPIDJSON_ERROR_STRING(x) TEXT(x)
 #endif
 
+#ifndef RAPIDJSON_MALLOC
+#define RAPIDJSON_MALLOC FMemory::Malloc
+#endif
+
+#ifndef RAPIDJSON_REALLOC
+#define RAPIDJSON_REALLOC FMemory::Realloc
+#endif
+
+#ifndef RAPIDJSON_FREE
+#define RAPIDJSON_FREE FMemory::Free
+#endif
+
+#if PLATFORM_LITTLE_ENDIAN
+#define RAPIDJSON_ENDIAN RAPIDJSON_LITTLEENDIAN
+#else
+#define RAPIDJSON_ENDIAN RAPIDJSON_BIGENDIAN
+#endif
+
+#if PLATFORM_CPU_ARM_FAMILY
+#define RAPIDJSON_NEON
+#elif PLATFORM_CPU_X86_FAMILY
+#define RAPIDJSON_SSE42
+#endif
+
 #include "rapidjson/rapidjson.h"
 #include "rapidjson/document.h"
 #include "rapidjson/encodings.h"
diff --git a/Engine/Source/Runtime/Projects/Private/CustomBuildSteps.cpp b/Engine/Source/Runtime/Projects/Private/CustomBuildSteps.cpp
--- a/Engine/Source/Runtime/Projects/Private/CustomBuildSteps.cpp
+++ b/Engine/Source/Runtime/Projects/Private/CustomBuildSteps.cpp
@@ -3,6 +3,8 @@
 #include "CustomBuildSteps.h"
 #include "Dom/JsonObject.h"
 #include "Serialization/JsonSerializer.h"
+#include "JsonUtils/JsonConversion.h"
+#include "RapidJsonPluginLoading.h"
 
 #define LOCTEXT_NAMESPACE "ModuleDescriptor"
 
@@ -11,30 +13,47 @@ bool FCustomBuildSteps::IsEmpty() const
 	return HostPlatformToCommands.Num() == 0;
 }
 
-void FCustomBuildSteps::Read(const FJsonObject& Object, const FString& FieldName)
+FCustomBuildSteps UE::Projects::Private::ReadCustomBuildSteps(UE::Json::FConstObject Object, const TCHAR* FieldName)
 {
-	TSharedPtr<FJsonValue> StepsValue = Object.TryGetField(FieldName);
-	if(StepsValue.IsValid() && StepsValue->Type == EJson::Object)
+	FCustomBuildSteps Result{};
+
+	TOptional<Json::FConstObject> StepsObject = UE::Json::GetObjectField(Object, FieldName);
+	if (StepsObject.IsSet())
 	{
-		const TSharedPtr<FJsonObject>& StepsObject = StepsValue->AsObject();
-		for(const TPair<FString, TSharedPtr<FJsonValue>>& HostPlatformAndSteps : StepsObject->Values)
+		for (const Json::FMember& Member : *StepsObject)
 		{
-			TArray<FString>& Commands = HostPlatformToCommands.FindOrAdd(HostPlatformAndSteps.Key);
-			if(HostPlatformAndSteps.Value.IsValid() && HostPlatformAndSteps.Value->Type == EJson::Array)
+			TArray<FString>& Commands = Result.HostPlatformToCommands.FindOrAdd(Member.name.GetString());
+			if (Member.value.IsArray())
 			{
-				const TArray<TSharedPtr<FJsonValue>>& CommandsArray = HostPlatformAndSteps.Value->AsArray();
-				for(const TSharedPtr<FJsonValue>& CommandValue: CommandsArray)
+				Json::FConstArray CommandsArray = Member.value.GetArray();
+				Commands.Reserve(CommandsArray.Size());
+				for (const Json::FValue& CommandValue : CommandsArray)
 				{
-					if(CommandValue->Type == EJson::String)
+					if (CommandValue.IsString())
 					{
-						Commands.Add(CommandValue->AsString());
+						Commands.Emplace(CommandValue.GetString());
 					}
 				}
 			}
 		}
 	}
+
+	return Result;
+}
+
+void FCustomBuildSteps::Read(const FJsonObject& Object, const FString& FieldName)
+{
+	UE::Projects::Private::ReadFromDefaultJsonHelper(
+		Object,
+		[&FieldName, this](UE::Json::FConstObject RootObject)
+		{
+			*this = UE::Projects::Private::ReadCustomBuildSteps(RootObject, *FieldName);
+			return true;
+		}
+	);
 }
 
+
 void FCustomBuildSteps::Write(TJsonWriter<>& Writer, const FString& FieldName) const
 {
 	if (!IsEmpty())
diff --git a/Engine/Source/Runtime/Projects/Private/LocalizationDescriptor.cpp b/Engine/Source/Runtime/Projects/Private/LocalizationDescriptor.cpp
--- a/Engine/Source/Runtime/Projects/Private/LocalizationDescriptor.cpp
+++ b/Engine/Source/Runtime/Projects/Private/LocalizationDescriptor.cpp
@@ -6,6 +6,8 @@
 #include "Dom/JsonObject.h"
 #include "Serialization/JsonSerializer.h"
 #include "JsonUtils/JsonObjectArrayUpdater.h"
+#include "JsonUtils/JsonConversion.h"
+#include "RapidJsonPluginLoading.h"
 
 #define LOCTEXT_NAMESPACE "LocalizationDescriptor"
 
@@ -109,98 +111,58 @@ FLocalizationTargetDescriptor::FLocalizationTargetDescriptor(FString InName, ELo
 {
 }
 
-bool FLocalizationTargetDescriptor::Read(const FJsonObject& InObject, FText* OutFailReason /*= nullptr*/)
+TOptional<FText> UE::Projects::Private::Read(UE::Json::FConstObject InObject, FLocalizationTargetDescriptor& Result)
 {
 	// Read the target name
-	TSharedPtr<FJsonValue> NameValue = InObject.TryGetField(TEXT("Name"));
-	if (!NameValue.IsValid() || NameValue->Type != EJson::String)
+	if (!TryGetStringField(InObject, TEXT("Name"), Result.Name))
 	{
-		if (OutFailReason)
-		{
-			*OutFailReason = LOCTEXT("TargetWithoutAName", "Found a 'Localization Target' entry with a missing 'Name' field");
+		return LOCTEXT("TargetWithoutAName", "Found a 'Localization Target' entry with a missing 'Name' field");
 		}
-		return false;
-	}
-	Name = NameValue->AsString();
 
 	// Read the target loading policy
-	TSharedPtr<FJsonValue> LoadingPolicyValue = InObject.TryGetField(TEXT("LoadingPolicy"));
-	if (LoadingPolicyValue.IsValid() && LoadingPolicyValue->Type == EJson::String)
+	FString LoadingPolicyValue;
+	if (TryGetStringField(InObject, TEXT("LoadingPolicy"), LoadingPolicyValue))
 	{
-		LoadingPolicy = ELocalizationTargetDescriptorLoadingPolicy::FromString(*LoadingPolicyValue->AsString());
-		if (LoadingPolicy == ELocalizationTargetDescriptorLoadingPolicy::Max)
+		Result.LoadingPolicy = ELocalizationTargetDescriptorLoadingPolicy::FromString(*LoadingPolicyValue);
+		if (Result.LoadingPolicy == ELocalizationTargetDescriptorLoadingPolicy::Max)
 		{
-			if (OutFailReason)
-			{
-				*OutFailReason = FText::Format(LOCTEXT("TargetWithInvalidLoadingPolicy", "Localization Target entry '{0}' specified an unrecognized target LoadingPolicy '{1}'"), FText::FromString(Name), FText::FromString(LoadingPolicyValue->AsString()));
-			}
-			return false;
+			return FText::Format(LOCTEXT("TargetWithInvalidLoadingPolicy", "Localization Target entry '{0}' specified an unrecognized target LoadingPolicy '{1}'"), FText::FromString(Result.Name), FText::FromString(LoadingPolicyValue));
 		}
 	}
 
 		// Read the target generation policy
-	TSharedPtr<FJsonValue> ConfigGenerationPolicyValue = InObject.TryGetField(TEXT("ConfigGenerationPolicy"));
-	if (ConfigGenerationPolicyValue.IsValid() && ConfigGenerationPolicyValue->Type == EJson::String)
+	FString ConfigGenerationPolicyValue;
+	if (TryGetStringField(InObject, TEXT("ConfigGenerationPolicy"), ConfigGenerationPolicyValue))
 	{
-		ConfigGenerationPolicy = ELocalizationConfigGenerationPolicy::FromString(*ConfigGenerationPolicyValue->AsString());
-		if (ConfigGenerationPolicy == ELocalizationConfigGenerationPolicy::Max)
+		Result.ConfigGenerationPolicy = ELocalizationConfigGenerationPolicy::FromString(*ConfigGenerationPolicyValue);
+		if (Result.ConfigGenerationPolicy == ELocalizationConfigGenerationPolicy::Max)
 		{
-			if (OutFailReason)
-			{
-				*OutFailReason = FText::Format(LOCTEXT("TargetWithInvalidGenerationPolicy", "Localization Target entry '{0}' specified an unrecognized target GenerationPolicy'{1}'"), FText::FromString(Name), FText::FromString(ConfigGenerationPolicyValue->AsString()));
+			return FText::Format(LOCTEXT("TargetWithInvalidGenerationPolicy", "Localization Target entry '{0}' specified an unrecognized target GenerationPolicy'{1}'"), FText::FromString(Result.Name), FText::FromString(ConfigGenerationPolicyValue));
 			}
-			return false;
 		}
-	}
 	// If we can't find the value, GenerationPolicy already defaults to Never so it's ok 
 
-	return true;
+	return {};
+}
+
+bool FLocalizationTargetDescriptor::Read(const FJsonObject& InObject, FText* OutFailReason /*= nullptr*/)
+{
+	return UE::Projects::Private::ReadFromDefaultJson(InObject, *this, OutFailReason);
 }
 
 bool FLocalizationTargetDescriptor::Read(const FJsonObject& InObject, FText& OutFailReason)
 {
-	return Read(InObject, &OutFailReason);
+	return UE::Projects::Private::ReadFromDefaultJson(InObject, *this, &OutFailReason);
 }
 
 bool FLocalizationTargetDescriptor::ReadArray(const FJsonObject& InObject, const TCHAR* InName, TArray<FLocalizationTargetDescriptor>& OutTargets, FText* OutFailReason /*= nullptr*/)
 {
-	bool bResult = true;
-
-	TSharedPtr<FJsonValue> TargetsArrayValue = InObject.TryGetField(InName);
-	if (TargetsArrayValue.IsValid() && TargetsArrayValue->Type == EJson::Array)
-	{
-		const TArray<TSharedPtr<FJsonValue>>& TargetsArray = TargetsArrayValue->AsArray();
-		for (const TSharedPtr<FJsonValue>& TargetValue : TargetsArray)
-		{
-			if (TargetValue.IsValid() && TargetValue->Type == EJson::Object)
-			{
-				FLocalizationTargetDescriptor Descriptor;
-				if (Descriptor.Read(*TargetValue->AsObject().Get(), OutFailReason))
-				{
-					OutTargets.Add(Descriptor);
-				}
-				else
-				{
-					bResult = false;
-				}
-			}
-			else
-			{
-				if (OutFailReason)
-				{
-					*OutFailReason = LOCTEXT("TargetWithInvalidTargetsArray", "The 'Localization Targets' array has invalid contents and was not able to be loaded.");
-				}
-				bResult = false;
-			}
-		}
-	}
-
-	return bResult;
+	return UE::Projects::Private::ReadArrayFromDefaultJson(InObject, InName, OutTargets, OutFailReason);
 }
 
 bool FLocalizationTargetDescriptor::ReadArray(const FJsonObject& InObject, const TCHAR* InName, TArray<FLocalizationTargetDescriptor>& OutTargets, FText& OutFailReason)
 {
-	return ReadArray(InObject, InName, OutTargets, &OutFailReason);
+	return UE::Projects::Private::ReadArrayFromDefaultJson(InObject, InName, OutTargets, &OutFailReason);
 }
 
 void FLocalizationTargetDescriptor::Write(TJsonWriter<>& Writer) const
diff --git a/Engine/Source/Runtime/Projects/Private/ModuleDescriptor.cpp b/Engine/Source/Runtime/Projects/Private/ModuleDescriptor.cpp
--- a/Engine/Source/Runtime/Projects/Private/ModuleDescriptor.cpp
+++ b/Engine/Source/Runtime/Projects/Private/ModuleDescriptor.cpp
@@ -7,7 +7,8 @@
 #include "Serialization/JsonSerializer.h"
 #include "JsonUtils/JsonObjectArrayUpdater.h"
 #include "Modules/ModuleManager.h"
-#include "JsonExtensions.h"
+#include "RapidJsonPluginLoading.h"
+#include "JsonUtils/JsonConversion.h"
 
 #define LOCTEXT_NAMESPACE "ModuleDescriptor"
 
@@ -157,158 +158,118 @@ FModuleDescriptor::FModuleDescriptor(const FName InName, EHostType::Type InType,
 {
 }
 
-bool FModuleDescriptor::Read(const FJsonObject& Object, FText* OutFailReason /*= nullptr*/)
+
+TOptional<FText> UE::Projects::Private::Read(UE::Json::FConstObject Object, FModuleDescriptor& Out)
 {
 	// Read the module name
-	TSharedPtr<FJsonValue> NameValue = Object.TryGetField(TEXT("Name"));
-	if(!NameValue.IsValid() || NameValue->Type != EJson::String)
+	FString NameValue;
+	if (!TryGetStringField(Object, TEXT("Name"), NameValue))
 	{
-		if (OutFailReason)
-		{
-			*OutFailReason = LOCTEXT("ModuleWithoutAName", "Found a 'Module' entry with a missing 'Name' field");
-		}
-		return false;
+		return LOCTEXT("ModuleWithoutAName", "Found a 'Module' entry with a missing 'Name' field");
 	}
-	Name = FName(*NameValue->AsString());
+	Out.Name = FName(*NameValue);
 
 	// Read the module type
-	TSharedPtr<FJsonValue> TypeValue = Object.TryGetField(TEXT("Type"));
-	if(!TypeValue.IsValid() || TypeValue->Type != EJson::String)
+	FString TypeValue;
+	if (!TryGetStringField(Object, TEXT("Type"), TypeValue))
 	{
-		if (OutFailReason)
-		{
-			*OutFailReason = FText::Format( LOCTEXT( "ModuleWithoutAType", "Found Module entry '{0}' with a missing 'Type' field" ), FText::FromName(Name) );
-		}
-		return false;
+		return FText::Format( LOCTEXT( "ModuleWithoutAType", "Found Module entry '{0}' with a missing 'Type' field" ), FText::FromName(Out.Name));
 	}
-	Type = EHostType::FromString(*TypeValue->AsString());
-	if(Type == EHostType::Max)
+
+	Out.Type = EHostType::FromString(*TypeValue);
+	if (Out.Type == EHostType::Max)
 	{
-		if (OutFailReason)
-		{
-			*OutFailReason = FText::Format( LOCTEXT( "ModuleWithInvalidType", "Module entry '{0}' specified an unrecognized module Type '{1}'" ), FText::FromName(Name), FText::FromString(TypeValue->AsString()) );
-		}
-		return false;
+		return FText::Format( LOCTEXT( "ModuleWithInvalidType", "Module entry '{0}' specified an unrecognized module Type '{1}'" ), FText::FromName(Out.Name), FText::FromString(TypeValue) );
 	}
 
 	// Read the loading phase
-	TSharedPtr<FJsonValue> LoadingPhaseValue = Object.TryGetField(TEXT("LoadingPhase"));
-	if(LoadingPhaseValue.IsValid() && LoadingPhaseValue->Type == EJson::String)
+	FString LoadingPhaseValue;
+	if (TryGetStringField(Object, TEXT("LoadingPhase"), LoadingPhaseValue))
 	{
-		LoadingPhase = ELoadingPhase::FromString(*LoadingPhaseValue->AsString());
-		if(LoadingPhase == ELoadingPhase::Max)
+		Out.LoadingPhase = ELoadingPhase::FromString(*LoadingPhaseValue);
+		if(Out.LoadingPhase == ELoadingPhase::Max)
 		{
-			if (OutFailReason)
-			{
-				*OutFailReason = FText::Format( LOCTEXT( "ModuleWithInvalidLoadingPhase", "Module entry '{0}' specified an unrecognized module LoadingPhase '{1}'" ), FText::FromName(Name), FText::FromString(LoadingPhaseValue->AsString()) );
-			}
-			return false;
+			return FText::Format( LOCTEXT( "ModuleWithInvalidLoadingPhase", "Module entry '{0}' specified an unrecognized module LoadingPhase '{1}'" ), FText::FromName(Out.Name), FText::FromString(LoadingPhaseValue) );
 		}
 	}
 
 	// Read the allow/deny lists for platforms
-	JsonExtensions::TryGetStringArrayFieldWithDeprecatedFallback(Object, TEXT("PlatformAllowList"), TEXT("WhitelistPlatforms"), /*out*/ PlatformAllowList);
-	JsonExtensions::TryGetStringArrayFieldWithDeprecatedFallback(Object, TEXT("PlatformDenyList"), TEXT("BlacklistPlatforms"), /*out*/ PlatformDenyList);
-	Object.TryGetBoolField(TEXT("HasExplicitPlatforms"), bHasExplicitPlatforms);
+	TryGetStringArrayFieldWithDeprecatedFallback(Object, TEXT("PlatformAllowList"), TEXT("WhitelistPlatforms"), /*out*/ Out.PlatformAllowList);
+	TryGetStringArrayFieldWithDeprecatedFallback(Object, TEXT("PlatformDenyList"), TEXT("BlacklistPlatforms"), /*out*/ Out.PlatformDenyList);
+	TryGetBoolField(Object, TEXT("HasExplicitPlatforms"), Out.bHasExplicitPlatforms);
 
 	// Read the allow/deny lists for targets
-	JsonExtensions::TryGetEnumArrayFieldWithDeprecatedFallback(Object, TEXT("TargetAllowList"), TEXT("WhitelistTargets"), /*out*/ TargetAllowList);
-	JsonExtensions::TryGetEnumArrayFieldWithDeprecatedFallback(Object, TEXT("TargetDenyList"), TEXT("BlacklistTargets"), /*out*/ TargetDenyList);
+	TryGetEnumArrayFieldWithDeprecatedFallback(Object, TEXT("TargetAllowList"), TEXT("WhitelistTargets"), /*out*/ Out.TargetAllowList);
+	TryGetEnumArrayFieldWithDeprecatedFallback(Object, TEXT("TargetDenyList"), TEXT("BlacklistTargets"), /*out*/ Out.TargetDenyList);
 
 	// Read the allow/deny lists for target configurations
-	JsonExtensions::TryGetEnumArrayFieldWithDeprecatedFallback(Object, TEXT("TargetConfigurationAllowList"), TEXT("WhitelistTargetConfigurations"), /*out*/ TargetConfigurationAllowList);
-	JsonExtensions::TryGetEnumArrayFieldWithDeprecatedFallback(Object, TEXT("TargetConfigurationDenyList"), TEXT("BlacklistTargetConfigurations"), /*out*/ TargetConfigurationDenyList);
+	TryGetEnumArrayFieldWithDeprecatedFallback(Object, TEXT("TargetConfigurationAllowList"), TEXT("WhitelistTargetConfigurations"), /*out*/ Out.TargetConfigurationAllowList);
+	TryGetEnumArrayFieldWithDeprecatedFallback(Object, TEXT("TargetConfigurationDenyList"), TEXT("BlacklistTargetConfigurations"), /*out*/ Out.TargetConfigurationDenyList);
 
 	// Read the allow/deny lists for programs
-	JsonExtensions::TryGetStringArrayFieldWithDeprecatedFallback(Object, TEXT("ProgramAllowList"), TEXT("WhitelistPrograms"), /*out*/ ProgramAllowList);
-	JsonExtensions::TryGetStringArrayFieldWithDeprecatedFallback(Object, TEXT("ProgramDenyList"), TEXT("BlacklistPrograms"), /*out*/ ProgramDenyList);
+	TryGetStringArrayFieldWithDeprecatedFallback(Object, TEXT("ProgramAllowList"), TEXT("WhitelistPrograms"), /*out*/ Out.ProgramAllowList);
+	TryGetStringArrayFieldWithDeprecatedFallback(Object, TEXT("ProgramDenyList"), TEXT("BlacklistPrograms"), /*out*/ Out.ProgramDenyList);
 
 	// Read the additional dependencies
-	Object.TryGetStringArrayField(TEXT("AdditionalDependencies"), AdditionalDependencies);
+	TryGetStringArrayField(Object, TEXT("AdditionalDependencies"), Out.AdditionalDependencies);
 	
 	// Read the allow/deny list for game targets
-	Object.TryGetStringArrayField(TEXT("GameTargetAllowList"), GameTargetAllowList);
-	Object.TryGetStringArrayField(TEXT("GameTargetDenyList"), GameTargetDenyList);
+	TryGetStringArrayField(Object, TEXT("GameTargetAllowList"), Out.GameTargetAllowList);
+	TryGetStringArrayField(Object, TEXT("GameTargetDenyList"), Out.GameTargetDenyList);
 	
 	// Read the platform architecture allow/deny lists
-	auto ReadPlatformArchitectureList = [this, Object, OutFailReason]( const TCHAR* FieldName, TMap<FString,TArray<FString>>& OutPlatformArchitectureList )
+	auto ReadPlatformArchitectureList = [&Out, Object]( const TCHAR* FieldName, TMap<FString,TArray<FString>>& OutPlatformArchitectureList ) -> TOptional<FText>
 	{
 		TArray<FString> PlatformArchitecturePairs;
-		Object.TryGetStringArrayField(FieldName, PlatformArchitecturePairs);
+		TryGetStringArrayField(Object, FieldName, PlatformArchitecturePairs);
 		for (const FString& PlatformArchitecturePair : PlatformArchitecturePairs)
 		{
 			FString Platform, Architecture;
 			if (!PlatformArchitecturePair.Split(TEXT(":"), &Platform, &Architecture))
 			{
-				if (OutFailReason)
-				{
-					*OutFailReason = FText::Format( LOCTEXT( "ModuleWithInvalidPlatformArchitecture", "Module entry '{0}' specified an unrecognized Platform:Architecture pair '{1}'" ), FText::FromName(Name), FText::FromString(PlatformArchitecturePair) );
-				}
-				return false;
+				return FText::Format( LOCTEXT( "ModuleWithInvalidPlatformArchitecture", "Module entry '{0}' specified an unrecognized Platform:Architecture pair '{1}'" ), FText::FromName(Out.Name), FText::FromString(PlatformArchitecturePair) );
 			}
 
 			OutPlatformArchitectureList.FindOrAdd(Platform).AddUnique(Architecture);
 		}
-		return true;
 
+		return {};
 	};
-	if (!ReadPlatformArchitectureList(TEXT("PlatformArchitectureAllowList"), PlatformArchitectureAllowList))
+
+	TOptional<FText> ErrorResult = ReadPlatformArchitectureList(TEXT("PlatformArchitectureAllowList"), Out.PlatformArchitectureAllowList);
+	if (ErrorResult)
 	{
-		return false;
+		return ErrorResult;
 	}
-	if (!ReadPlatformArchitectureList(TEXT("PlatformArchitectureDenyList"), PlatformArchitectureDenyList))
+
+	ErrorResult = ReadPlatformArchitectureList(TEXT("PlatformArchitectureDenyList"), Out.PlatformArchitectureDenyList);
+
+	if (ErrorResult)
 	{
-		return false;
+		return ErrorResult;
 	}
 
-	return true;
+	return {};
+}
+	
+bool FModuleDescriptor::Read(const FJsonObject& Object, FText* OutFailReason /*= nullptr*/)
+{
+	return UE::Projects::Private::ReadFromDefaultJson(Object, *this, OutFailReason);
 }
 
 bool FModuleDescriptor::Read(const FJsonObject& Object, FText& OutFailReason)
 {
-	return Read(Object, &OutFailReason);
+	return UE::Projects::Private::ReadFromDefaultJson(Object, *this, &OutFailReason);
 }
 
 bool FModuleDescriptor::ReadArray(const FJsonObject& Object, const TCHAR* Name, TArray<FModuleDescriptor>& OutModules, FText* OutFailReason /*= nullptr*/)
 {
-	bool bResult = true;
-
-	TSharedPtr<FJsonValue> ModulesArrayValue = Object.TryGetField(Name);
-	if(ModulesArrayValue.IsValid() && ModulesArrayValue->Type == EJson::Array)
-	{
-		const TArray< TSharedPtr< FJsonValue > >& ModulesArray = ModulesArrayValue->AsArray();
-		for(int Idx = 0; Idx < ModulesArray.Num(); Idx++)
-		{
-			const TSharedPtr<FJsonValue>& ModuleValue = ModulesArray[Idx];
-			if(ModuleValue.IsValid() && ModuleValue->Type == EJson::Object)
-			{
-				FModuleDescriptor Descriptor;
-				if(Descriptor.Read(*ModuleValue->AsObject().Get(), OutFailReason))
-				{
-					OutModules.Add(Descriptor);
-				}
-				else
-				{
-					bResult = false;
-				}
-			}
-			else
-			{
-				if (OutFailReason)
-				{
-					*OutFailReason = LOCTEXT( "ModuleWithInvalidModulesArray", "The 'Modules' array has invalid contents and was not able to be loaded." );
-				}
-				bResult = false;
-			}
-		}
-	}
-	
-	return bResult;
+	return UE::Projects::Private::ReadArrayFromDefaultJson(Object, Name, OutModules, OutFailReason);
 }
 
 bool FModuleDescriptor::ReadArray(const FJsonObject& Object, const TCHAR* Name, TArray<FModuleDescriptor>& OutModules, FText& OutFailReason)
 {
-	return ReadArray(Object, Name, OutModules, &OutFailReason);
+	return UE::Projects::Private::ReadArrayFromDefaultJson(Object, Name, OutModules, &OutFailReason);
 }
 
 void FModuleDescriptor::Write(TJsonWriter<>& Writer) const
diff --git a/Engine/Source/Runtime/Projects/Private/PluginDescriptor.cpp b/Engine/Source/Runtime/Projects/Private/PluginDescriptor.cpp
--- a/Engine/Source/Runtime/Projects/Private/PluginDescriptor.cpp
+++ b/Engine/Source/Runtime/Projects/Private/PluginDescriptor.cpp
@@ -6,6 +6,8 @@
 #include "Serialization/JsonReader.h"
 #include "Serialization/JsonSerializer.h"
 #include "ProjectDescriptor.h"
+#include "JsonUtils/JsonConversion.h"
+#include "RapidJsonPluginLoading.h"
 
 #define LOCTEXT_NAMESPACE "PluginDescriptor"
 
@@ -122,26 +124,48 @@ bool FPluginDescriptor::Load(const FString& FileName, FText& OutFailReason)
 
 bool FPluginDescriptor::Read(const FString& Text, FText* OutFailReason /*= nullptr*/)
 {
+	using namespace UE::Json;
+
 #if WITH_EDITOR
 	CachedJson.Reset();
 	AdditionalFieldsToWrite.Reset();
 	AdditionalFieldsToRemove.Reset();
 #endif // WITH_EDITOR
 
+	TValueOrError<FDocument, FParseError> ParseResult = Parse(Text);
+
 	// Deserialize a JSON object from the string
-	TSharedPtr<FJsonObject> JsonObject = PluginDescriptor::DeserializeJson(Text, OutFailReason);
-	if (JsonObject.IsValid())
+	if (ParseResult.HasError())
+	{
+		if (OutFailReason)
+	{
+			*OutFailReason = FText::Format(LOCTEXT("FailedToReadDescriptorFile", "Failed to read file. {0}"), FText::FromString(ParseResult.GetError().CreateMessage(Text)));
+		}
+		return false;
+	}
+
+	FDocument Document = ParseResult.StealValue();
+
+	TOptional<FConstObject> RootObject = GetRootObject(Document);
+	if (RootObject.IsSet())
 	{
 		// Parse it as a plug-in descriptor
-		if (Read(*JsonObject.Get(), OutFailReason))
+		if (TOptional<FText> ReadError = UE::Projects::Private::Read(*RootObject, *this))
+		{
+			if (OutFailReason)
 		{
+				*OutFailReason = MoveTemp(*ReadError);
+			}			
+			return false;
+		}
+
 #if WITH_EDITOR
-			CachedJson = JsonObject;
+		CachedJson = ConvertRapidJsonToSharedJsonObject(*RootObject);
 			AdditionalFieldsToWrite.Reset();
 			AdditionalFieldsToRemove.Reset();
 #endif // WITH_EDITOR
 			return true;
-		}
+
 	}
 	return false;
 }
@@ -151,19 +175,15 @@ bool FPluginDescriptor::Read(const FString& Text, FText& OutFailReason)
 	return Read(Text, &OutFailReason);
 }
 
-bool FPluginDescriptor::Read(const FJsonObject& Object, FText* OutFailReason /*= nullptr*/)
+TOptional<FText> UE::Projects::Private::Read(UE::Json::FConstObject Object, FPluginDescriptor& Result)
 {
 	// Read the file version
 	int32 FileVersionInt32;
-	if(!Object.TryGetNumberField(TEXT("FileVersion"), FileVersionInt32))
+	if(!TryGetNumberField(Object, TEXT("FileVersion"), FileVersionInt32))
 	{
-		if(!Object.TryGetNumberField(TEXT("PluginFileVersion"), FileVersionInt32))
+		if(!TryGetNumberField(Object, TEXT("PluginFileVersion"), FileVersionInt32))
 		{
-			if (OutFailReason)
-			{
-				*OutFailReason = LOCTEXT("InvalidProjectFileVersion", "File does not have a valid 'FileVersion' number.");
-			}
-			return false;
+			return LOCTEXT("InvalidProjectFileVersion", "File does not have a valid 'FileVersion' number.");
 		}
 	}
 
@@ -171,155 +191,149 @@ bool FPluginDescriptor::Read(const FJsonObject& Object, FText* OutFailReason /*=
 	EPluginDescriptorVersion PluginFileVersion = (EPluginDescriptorVersion)FileVersionInt32;
 	if ((PluginFileVersion <= EPluginDescriptorVersion::Invalid) || (PluginFileVersion > EPluginDescriptorVersion::Latest))
 	{
-		if (OutFailReason)
-		{
 			const FText ReadVersionText = FText::FromString(FString::Printf(TEXT("%d"), (int32)PluginFileVersion));
 			const FText LatestVersionText = FText::FromString(FString::Printf(TEXT("%d"), (int32)EPluginDescriptorVersion::Latest));
-			*OutFailReason = FText::Format(LOCTEXT("ProjectFileVersionTooLarge", "File appears to be in a newer version ({0}) of the file format that we can load (max version: {1})."), ReadVersionText, LatestVersionText);
-		}
-		return false;
+		return FText::Format(LOCTEXT("ProjectFileVersionTooLarge", "File appears to be in a newer version ({0}) of the file format that we can load (max version: {1})."), ReadVersionText, LatestVersionText);
 	}
 
 	// Read the other fields
-	Object.TryGetNumberField(TEXT("Version"), Version);
-	Object.TryGetStringField(TEXT("VersionName"), VersionName);
-	Object.TryGetStringField(TEXT("FriendlyName"), FriendlyName);
-	Object.TryGetStringField(TEXT("Description"), Description);
+	TryGetNumberField(Object, TEXT("Version"), Result.Version);
+	TryGetStringField(Object, TEXT("VersionName"), Result.VersionName);
+	TryGetStringField(Object, TEXT("FriendlyName"), Result.FriendlyName);
+	TryGetStringField(Object, TEXT("Description"), Result.Description);
 
-	if (!Object.TryGetStringField(TEXT("Category"), Category))
+	if (!TryGetStringField(Object, TEXT("Category"), Result.Category))
 	{
 		// Category used to be called CategoryPath in .uplugin files
-		Object.TryGetStringField(TEXT("CategoryPath"), Category);
+		TryGetStringField(Object, TEXT("CategoryPath"), Result.Category);
 	}
         
 	// Due to a difference in command line parsing between Windows and Mac, we shipped a few Mac samples containing
 	// a category name with escaped quotes. Remove them here to make sure we can list them in the right category.
-	if (Category.Len() >= 2 && Category.StartsWith(TEXT("\""), ESearchCase::CaseSensitive) && Category.EndsWith(TEXT("\""), ESearchCase::CaseSensitive))
+	if (Result.Category.Len() >= 2 && Result.Category.StartsWith(TEXT("\""), ESearchCase::CaseSensitive) && Result.Category.EndsWith(TEXT("\""), ESearchCase::CaseSensitive))
 	{
-		Category.MidInline(1, Category.Len() - 2, EAllowShrinking::No);
+		Result.Category.MidInline(1, Result.Category.Len() - 2, EAllowShrinking::No);
 	}
 
-	Object.TryGetStringField(TEXT("CreatedBy"), CreatedBy);
-	Object.TryGetStringField(TEXT("CreatedByURL"), CreatedByURL);
-	Object.TryGetStringField(TEXT("DocsURL"), DocsURL);
-	Object.TryGetStringField(TEXT("MarketplaceURL"), MarketplaceURL);
-	Object.TryGetStringField(TEXT("SupportURL"), SupportURL);
-	Object.TryGetStringField(TEXT("EngineVersion"), EngineVersion);
-	Object.TryGetStringField(TEXT("EditorCustomVirtualPath"), EditorCustomVirtualPath);
-	Object.TryGetStringArrayField(TEXT("SupportedTargetPlatforms"), SupportedTargetPlatforms);
-	Object.TryGetStringArrayField(TEXT("SupportedPrograms"), SupportedPrograms);
-	Object.TryGetBoolField(TEXT("bIsPluginExtension"), bIsPluginExtension);
+	TryGetStringField(Object, TEXT("CreatedBy"), Result.CreatedBy);
+	TryGetStringField(Object, TEXT("CreatedByURL"), Result.CreatedByURL);
+	TryGetStringField(Object, TEXT("DocsURL"), Result.DocsURL);
+	TryGetStringField(Object, TEXT("MarketplaceURL"), Result.MarketplaceURL);
+	TryGetStringField(Object, TEXT("SupportURL"), Result.SupportURL);
+	TryGetStringField(Object, TEXT("EngineVersion"), Result.EngineVersion);
+	TryGetStringField(Object, TEXT("EditorCustomVirtualPath"), Result.EditorCustomVirtualPath);
+	TryGetStringArrayField(Object, TEXT("SupportedTargetPlatforms"), Result.SupportedTargetPlatforms);
+	TryGetStringArrayField(Object, TEXT("SupportedPrograms"), Result.SupportedPrograms);
+	TryGetBoolField(Object, TEXT("bIsPluginExtension"), Result.bIsPluginExtension);
 
-	if (!FModuleDescriptor::ReadArray(Object, TEXT("Modules"), Modules, OutFailReason))
+	if (TOptional<FText> ReadError = ReadArray(Object, TEXT("Modules"), Result.Modules))
 	{
-		return false;
+		return ReadError;
 	}
 
-	if (!FLocalizationTargetDescriptor::ReadArray(Object, TEXT("LocalizationTargets"), LocalizationTargets, OutFailReason))
+	if (TOptional<FText> ReadError = ReadArray(Object, TEXT("LocalizationTargets"), Result.LocalizationTargets))
 	{
-		return false;
+		return ReadError;
 	}
 
-	Object.TryGetStringField(TEXT("VersePath"), VersePath);
-	Object.TryGetStringField(TEXT("DeprecatedEngineVersion"), DeprecatedEngineVersion);
+	TryGetStringField(Object, TEXT("VersePath"), Result.VersePath);
+	TryGetStringField(Object, TEXT("DeprecatedEngineVersion"), Result.DeprecatedEngineVersion);
 
 	// Read the Verse scope
-	TSharedPtr<FJsonValue> VerseScopeValue = Object.TryGetField(TEXT("VerseScope"));
-	if (VerseScopeValue.IsValid() && VerseScopeValue->Type == EJson::String)
+	Json::FValue::ConstMemberIterator VerseScopeValue = Object.FindMember(TEXT("VerseScope"));
+	if (VerseScopeValue != Object.MemberEnd() && VerseScopeValue->value.IsString())
 	{
-		if (TOptional<EVerseScope::Type> MaybeVerseScope = EVerseScope::FromString(*VerseScopeValue->AsString()))
+		if (TOptional<EVerseScope::Type> MaybeVerseScope = EVerseScope::FromString(VerseScopeValue->value.GetString()))
 		{
-			VerseScope = *MaybeVerseScope;
+			Result.VerseScope = *MaybeVerseScope;
 		}
 		else
 		{
-			if (OutFailReason)
-			{
-				*OutFailReason = FText::Format(LOCTEXT("PluginWithInvalidVerseScope", "Plugin entry 'VerseScope' specified an unrecognized value '{1}'"), FText::FromString(VerseScopeValue->AsString()));
-			}
-			return false;
+			return FText::Format(LOCTEXT("PluginWithInvalidVerseScope", "Plugin entry 'VerseScope' specified an unrecognized value '{0}'"), FText::FromString(VerseScopeValue->value.GetString()));
 		}
 	}
 
 	// Read the Verse version.
-	TSharedPtr<FJsonValue> VerseVersionValue = Object.TryGetField(TEXT("VerseVersion"));
-	if (VerseVersionValue.IsValid())
+	Json::FValue::ConstMemberIterator VerseVersionValue = Object.FindMember(TEXT("VerseVersion"));
+	if (VerseVersionValue != Object.MemberEnd())
 	{
-		uint32 PluginVerseVersion;
-		if (VerseVersionValue->TryGetNumber(PluginVerseVersion))
+		if (VerseVersionValue->value.IsInt())
 		{
-			VerseVersion = PluginVerseVersion;
+			Result.VerseVersion = VerseVersionValue->value.GetInt();
 		}
 		else
 		{
-			if (OutFailReason)
-			{
-				*OutFailReason = FText::Format(LOCTEXT("PluginWithInvalidVerseVersion", "Plugin entry 'VerseVersion' specified an unrecognized value '{1}'"), FText::FromString(VerseVersionValue->AsString()));
-			}
-			return false;
+			return FText::Format(LOCTEXT("PluginWithInvalidVerseVersion", "Plugin entry 'VerseVersion' specified an unrecognized type '{0}'"), FText::FromString(Json::GetValueTypeName(VerseVersionValue->value)));
 		}
 	}
 
-	Object.TryGetBoolField(TEXT("EnableSceneGraph"), bEnableSceneGraph);
+	TryGetBoolField(Object, TEXT("EnableSceneGraph"), Result.bEnableSceneGraph);
 	
-	Object.TryGetBoolField(TEXT("EnableVerseAssetReflection"), bEnableVerseAssetReflection);
+	TryGetBoolField(Object, TEXT("EnableVerseAssetReflection"), Result.bEnableVerseAssetReflection);
 
-	Object.TryGetBoolField(TEXT("EnableIAD"), bEnableIAD);
+	TryGetBoolField(Object, TEXT("EnableIAD"), Result.bEnableIAD);
 
 	bool bEnabledByDefault;
-	if(Object.TryGetBoolField(TEXT("EnabledByDefault"), bEnabledByDefault))
+	if(TryGetBoolField(Object, TEXT("EnabledByDefault"), bEnabledByDefault))
 	{
-		EnabledByDefault = bEnabledByDefault ? EPluginEnabledByDefault::Enabled : EPluginEnabledByDefault::Disabled;
+		Result.EnabledByDefault = bEnabledByDefault ? EPluginEnabledByDefault::Enabled : EPluginEnabledByDefault::Disabled;
 	}
 
-	Object.TryGetBoolField(TEXT("CanContainContent"), bCanContainContent);
-	Object.TryGetBoolField(TEXT("CanContainVerse"), bCanContainVerse);
-	Object.TryGetBoolField(TEXT("NoCode"), bNoCode);
-	Object.TryGetBoolField(TEXT("IsBetaVersion"), bIsBetaVersion);
-	Object.TryGetBoolField(TEXT("IsExperimentalVersion"), bIsExperimentalVersion);
-	Object.TryGetBoolField(TEXT("Installed"), bInstalled);
-	Object.TryGetBoolField(TEXT("RequiresBuildPlatform"), bRequiresBuildPlatform);
-	Object.TryGetBoolField(TEXT("Hidden"), bIsHidden);
-	Object.TryGetBoolField(TEXT("Sealed"), bIsSealed);
-	Object.TryGetBoolField(TEXT("ExplicitlyLoaded"), bExplicitlyLoaded);
-	Object.TryGetBoolField(TEXT("HasExplicitPlatforms"), bHasExplicitPlatforms);
+	TryGetBoolField(Object, TEXT("CanContainContent"), Result.bCanContainContent);
+	TryGetBoolField(Object, TEXT("CanContainVerse"), Result.bCanContainVerse);
+	TryGetBoolField(Object, TEXT("NoCode"), Result.bNoCode);
+	TryGetBoolField(Object, TEXT("IsBetaVersion"), Result.bIsBetaVersion);
+	TryGetBoolField(Object, TEXT("IsExperimentalVersion"), Result.bIsExperimentalVersion);
+	TryGetBoolField(Object, TEXT("Installed"), Result.bInstalled);
+	TryGetBoolField(Object, TEXT("RequiresBuildPlatform"), Result.bRequiresBuildPlatform);
+	TryGetBoolField(Object, TEXT("Hidden"), Result.bIsHidden);
+	TryGetBoolField(Object, TEXT("Sealed"), Result.bIsSealed);
+	TryGetBoolField(Object, TEXT("ExplicitlyLoaded"), Result.bExplicitlyLoaded);
+	TryGetBoolField(Object, TEXT("HasExplicitPlatforms"), Result.bHasExplicitPlatforms);
 
 	bool bCanBeUsedWithUnrealHeaderTool;
-	if(Object.TryGetBoolField(TEXT("CanBeUsedWithUnrealHeaderTool"), bCanBeUsedWithUnrealHeaderTool) && bCanBeUsedWithUnrealHeaderTool)
+	if(TryGetBoolField(Object, TEXT("CanBeUsedWithUnrealHeaderTool"), bCanBeUsedWithUnrealHeaderTool) && bCanBeUsedWithUnrealHeaderTool)
 	{
-		SupportedPrograms.Add(TEXT("UnrealHeaderTool"));
+		Result.SupportedPrograms.Add(TEXT("UnrealHeaderTool"));
 	}
 
-	PreBuildSteps.Read(Object, TEXT("PreBuildSteps"));
-	PostBuildSteps.Read(Object, TEXT("PostBuildSteps"));
+	Result.PreBuildSteps = ReadCustomBuildSteps(Object, TEXT("PreBuildSteps"));
+	Result.PostBuildSteps = ReadCustomBuildSteps(Object, TEXT("PostBuildSteps"));
 
-	if (!FPluginReferenceDescriptor::ReadArray(Object, TEXT("Plugins"), Plugins, OutFailReason))
+	if (TOptional<FText> ReadError = ReadArray(Object, TEXT("Plugins"), Result.Plugins))
 	{
-		return false;
+		return ReadError;
 	}
 
 	// Backwards compatibility support
-	TArray<FString> DisallowedPluginNameStrings;
-	if (Object.TryGetStringArrayField(TEXT("DisallowedPlugins"), DisallowedPluginNameStrings))
+	if (TArray<FString> DisallowedPluginNameStrings; TryGetStringArrayField(Object, TEXT("DisallowedPlugins"), DisallowedPluginNameStrings))
 	{
-		DisallowedPlugins.Reserve(DisallowedPluginNameStrings.Num());
+		Result.DisallowedPlugins.Reserve(DisallowedPluginNameStrings.Num());
 		for (int32 Index = 0; Index < DisallowedPluginNameStrings.Num(); ++Index)
 		{
-			FPluginDisallowedDescriptor& PluginDisallowedDescriptor = DisallowedPlugins.AddDefaulted_GetRef();
+			FPluginDisallowedDescriptor& PluginDisallowedDescriptor = Result.DisallowedPlugins.AddDefaulted_GetRef();
 			PluginDisallowedDescriptor.Name = DisallowedPluginNameStrings[Index];
 		}
 	}
-	else if (!FPluginDisallowedDescriptor::ReadArray(Object, TEXT("DisallowedPlugins"), DisallowedPlugins, OutFailReason))
+	else
 	{
-		return false;
+		if (TOptional<FText> ReadError = ReadArray(Object, TEXT("DisallowedPlugins"), Result.DisallowedPlugins))
+		{
+			return ReadError;
+		}
 	}
 
-	return true;
+	return {};
+}
+
+bool FPluginDescriptor::Read(const FJsonObject& Object, FText* OutFailReason /*= nullptr*/)
+{
+	return UE::Projects::Private::ReadFromDefaultJson(Object, *this, OutFailReason);
 }
 
 bool FPluginDescriptor::Read(const FJsonObject& Object, FText& OutFailReason)
 {
-	return Read(Object, &OutFailReason);
+	return UE::Projects::Private::ReadFromDefaultJson(Object, *this, &OutFailReason);
 }
 
 bool FPluginDescriptor::Save(const TCHAR* FileName, FText* OutFailReason /*= nullptr*/) const
diff --git a/Engine/Source/Runtime/Projects/Private/PluginDisallowedDescriptor.cpp b/Engine/Source/Runtime/Projects/Private/PluginDisallowedDescriptor.cpp
--- a/Engine/Source/Runtime/Projects/Private/PluginDisallowedDescriptor.cpp
+++ b/Engine/Source/Runtime/Projects/Private/PluginDisallowedDescriptor.cpp
@@ -7,6 +7,8 @@
 #include "JsonUtils/JsonObjectArrayUpdater.h"
 #include "ProjectDescriptor.h"
 #include "JsonExtensions.h"
+#include "JsonUtils/JsonConversion.h"
+#include "RapidJsonPluginLoading.h"
 
 #define LOCTEXT_NAMESPACE "PluginDisallowedDescriptor"
 
@@ -28,25 +30,24 @@ namespace PluginDisallowedDescriptor
 	}
 }
 
-bool FPluginDisallowedDescriptor::Read(const FJsonObject& Object, FText* OutFailReason /*= nullptr*/)
+TOptional<FText> UE::Projects::Private::Read(UE::Json::FConstObject Object, FPluginDisallowedDescriptor& Result)
 {
-	bool bSuccess = true;
-
 	// Get the name
-	if (!Object.TryGetStringField(TEXT("Name"), Name))
+	if (!TryGetStringField(Object, TEXT("Name"), Result.Name))
 	{
-		if (OutFailReason)
-		{
-			*OutFailReason = LOCTEXT("FPluginDisallowedWithoutName", "Disallowed plugin must have a 'Name' field");
+		return LOCTEXT("FPluginDisallowedWithoutName", "Disallowed plugin must have a 'Name' field");
 		}
-		bSuccess = false;
-	}
 
 #if WITH_EDITOR
-	Object.TryGetStringField(TEXT("Comment"), Comment);
+	TryGetStringField(Object, TEXT("Comment"), Result.Comment);
 #endif //if WITH_EDITOR
 
-	return bSuccess;
+	return {};
+}
+
+bool FPluginDisallowedDescriptor::Read(const FJsonObject& Object, FText* OutFailReason /*= nullptr*/)
+{
+	return UE::Projects::Private::ReadFromDefaultJson(Object, *this, OutFailReason);
 }
 
 void FPluginDisallowedDescriptor::Write(TJsonWriter<>& Writer) const
@@ -69,29 +70,7 @@ void FPluginDisallowedDescriptor::UpdateJson(FJsonObject& JsonObject) const
 
 bool FPluginDisallowedDescriptor::ReadArray(const FJsonObject& Object, const TCHAR* Name, TArray<FPluginDisallowedDescriptor>& OutPlugins, FText* OutFailReason /*= nullptr*/)
 {
-	const TArray< TSharedPtr<FJsonValue> >* Array;
-
-	if (Object.TryGetArrayField(Name, Array))
-	{
-		for (const TSharedPtr<FJsonValue>& Item : *Array)
-		{
-			const TSharedPtr<FJsonObject>* ItemObject = nullptr;
-
-			if (Item.IsValid() && Item->TryGetObject(ItemObject) && ItemObject && ItemObject->IsValid())
-			{
-				FPluginDisallowedDescriptor DisallowedRef;
-
-				if (!DisallowedRef.Read(*ItemObject->Get(), OutFailReason))
-				{
-					return false;
-				}
-
-				OutPlugins.Add(MoveTemp(DisallowedRef));
-			}
-		}
-	}
-
-	return true;
+	return UE::Projects::Private::ReadArrayFromDefaultJson(Object, Name, OutPlugins, OutFailReason);
 }
 
 void FPluginDisallowedDescriptor::WriteArray(TJsonWriter<>& Writer, const TCHAR* ArrayName, const TArray<FPluginDisallowedDescriptor>& Plugins)
diff --git a/Engine/Source/Runtime/Projects/Private/PluginManifest.cpp b/Engine/Source/Runtime/Projects/Private/PluginManifest.cpp
--- a/Engine/Source/Runtime/Projects/Private/PluginManifest.cpp
+++ b/Engine/Source/Runtime/Projects/Private/PluginManifest.cpp
@@ -4,9 +4,78 @@
 #include "Misc/FileHelper.h"
 #include "Serialization/JsonReader.h"
 #include "Serialization/JsonSerializer.h"
+#include "JsonUtils/JsonConversion.h"
+#include "Misc/Paths.h"
+#include "RapidJsonPluginLoading.h"
 
 #define LOCTEXT_NAMESPACE "PluginManifest"
 
+namespace PluginManifest
+{
+	void DumpContents(const TArray<FPluginManifestEntry>& Contents, const FString& Filename)
+	{
+		TArray<FString> Lines;
+
+		for (const FPluginManifestEntry& Entry : Contents)
+		{
+			FString Line;
+			Entry.Descriptor.Write(Line);
+
+			Line = Entry.File + TEXT("\n") + Line;
+
+			Lines.Add(MoveTemp(Line));
+		}
+
+		Lines.Sort();
+
+		FString DumpFilename = FPaths::ProjectSavedDir() / Filename;
+		FFileHelper::SaveStringArrayToFile(Lines, *DumpFilename, FFileHelper::EEncodingOptions::ForceUTF8WithoutBOM);
+	}
+}
+
+TOptional<FText> UE::Projects::Private::Read(UE::Json::FConstObject Object, FPluginManifest& Result)
+{
+	using namespace UE::Json;
+
+	TOptional<FConstArray> JsonContents = GetArrayField(Object, TEXT("Contents"));
+	if (!JsonContents.IsSet())
+	{
+		return LOCTEXT("FailedToReadDescriptorFile", "No 'Contents' array found");
+	}
+
+	Result.Contents.Reserve(JsonContents->Size());
+
+	for (const FValue& JsonEntry : *JsonContents)
+	{
+		if (!JsonEntry.IsObject())
+		{
+			continue;
+		}
+
+		FPluginManifestEntry& ManifestEntry = Result.Contents.Emplace_GetRef();
+
+		FConstObject JsonEntryObject = JsonEntry.GetObject();
+		TOptional<FConstObject> JsonDescriptor = GetObjectField(JsonEntryObject, TEXT("Descriptor"));
+
+		if (!JsonDescriptor.IsSet())
+		{
+			return FText::Format(LOCTEXT("FailedToReadDescriptorFile", "Plugin descriptor index '{0}' missing 'Descriptor' field"), Result.Contents.Num()-1);
+		}
+
+		if (!TryGetStringField(JsonEntryObject, TEXT("File"), ManifestEntry.File))
+		{
+			return FText::Format(LOCTEXT("FailedToReadDescriptorFile", "Plugin descriptor index '{0}' missing 'File' field"), Result.Contents.Num()-1);
+		}
+
+		if (TOptional<FText> ReadError = UE::Projects::Private::Read(*JsonDescriptor, ManifestEntry.Descriptor))
+		{
+			return FText::Format(LOCTEXT("FailedToReadDescriptorFile", "Failed to parse plugin descriptor entry for file '{0}': {1}"), FText::FromString(ManifestEntry.File), *ReadError);
+		}
+	}
+
+	return {};
+}
+
 bool FPluginManifest::Load(const FString& FileName, FText& OutFailReason)
 {
 	// Read the file to a string
@@ -17,35 +86,44 @@ bool FPluginManifest::Load(const FString& FileName, FText& OutFailReason)
 		return false;
 	}
 
-	// Deserialize a JSON object from the string
-	TSharedPtr< FJsonObject > ObjectPtr;
-	TSharedRef< TJsonReader<> > Reader = TJsonReaderFactory<>::Create(FileContents);
-	if (!FJsonSerializer::Deserialize(Reader, ObjectPtr) || !ObjectPtr.IsValid() )
 	{
-		OutFailReason = FText::Format(LOCTEXT("FailedToReadDescriptorFile", "Failed to read file. {0}"), FText::FromString(Reader->GetErrorMessage()));
+		SCOPED_BOOT_TIMING("FPluginManifest::Load RapidJson")
+
+		// grab the characters so we can parse insitu safely (FString API says adding non-terminating zeros can be unsafe so just take the data)
+		TArray<TCHAR> Characters = MoveTemp(FileContents.GetCharArray());
+		TValueOrError<UE::Json::FDocument, UE::Json::FParseError> Document = UE::Json::ParseInPlace(Characters);
+
+		if (Document.HasError())
+		{
+			// have to re-parse the string because it was destructively altered
+			ensure(FFileHelper::LoadFileToString(FileContents, *FileName));
+			OutFailReason = FText::Format(LOCTEXT("FailedToReadDescriptorFile", "Failed to read file. {0}"), FText::FromString(Document.GetError().CreateMessage(FileContents)));
+			
 		return false;
 	}
 
-	// Parse the object
-	return Read(*ObjectPtr.Get(), OutFailReason);
-}
-
-bool FPluginManifest::Read(const FJsonObject& Object, FText& OutFailReason)
-{
-	TArray<TSharedPtr<FJsonValue>> JsonContents = Object.GetArrayField(TEXT("Contents"));
-	for (const TSharedPtr<FJsonValue>& JsonEntryValue : JsonContents)
+		TOptional<UE::Json::FConstObject> RootObject = UE::Json::GetRootObject(Document.GetValue());
+		if (RootObject.IsSet())
+		{
+			//for (int32 N = 0; N < 1000; ++N)
 	{
-		TSharedPtr<FJsonObject> JsonEntry = JsonEntryValue->AsObject();
-
-		FPluginManifestEntry Entry;
-		Entry.File = JsonEntry->GetStringField(TEXT("File"));
-		if (!Entry.Descriptor.Read(*JsonEntry->GetObjectField(TEXT("Descriptor")).Get(), OutFailReason))
+				if (TOptional<FText> Result = UE::Projects::Private::Read(*RootObject, *this))
 		{
+					OutFailReason = *Result;
 			return false;
 		}
-		Contents.Add(Entry);
 	}
+
 	return true;
+		}
+	}
+
+	return false;
+}
+
+bool FPluginManifest::Read(const FJsonObject& Object, FText& OutFailReason)
+{
+	return UE::Projects::Private::ReadFromDefaultJson(Object, *this, &OutFailReason);
 }
 
 #undef LOCTEXT_NAMESPACE
diff --git a/Engine/Source/Runtime/Projects/Private/PluginReferenceDescriptor.cpp b/Engine/Source/Runtime/Projects/Private/PluginReferenceDescriptor.cpp
--- a/Engine/Source/Runtime/Projects/Private/PluginReferenceDescriptor.cpp
+++ b/Engine/Source/Runtime/Projects/Private/PluginReferenceDescriptor.cpp
@@ -6,7 +6,8 @@
 #include "Serialization/JsonSerializer.h"
 #include "JsonUtils/JsonObjectArrayUpdater.h"
 #include "ProjectDescriptor.h"
-#include "JsonExtensions.h"
+#include "RapidJsonPluginLoading.h"
+#include "JsonUtils/JsonConversion.h"
 
 #define LOCTEXT_NAMESPACE "PluginDescriptor"
 
@@ -118,121 +119,88 @@ bool FPluginReferenceDescriptor::IsSupportedTargetPlatform(const FString& Platfo
 	}
 }
 
-bool FPluginReferenceDescriptor::Read(const TSharedRef<FJsonObject>& Object, FText* OutFailReason /*= nullptr*/)
+TOptional<FText> UE::Projects::Private::Read(UE::Json::FConstObject Object, FPluginReferenceDescriptor& Result)
 {
-PRAGMA_DISABLE_DEPRECATION_WARNINGS
-	return Read(*Object, OutFailReason, Object);
-PRAGMA_ENABLE_DEPRECATION_WARNINGS
-}
+	using namespace UE::Json;
 
-bool FPluginReferenceDescriptor::Read(const FJsonObject& Object, FText* OutFailReason /*= nullptr*/, TSharedPtr<FJsonObject> ObjectPtr /*= nullptr*/)
-{
 #if WITH_EDITOR
-	CachedJson = ObjectPtr;
-	AdditionalFieldsToWrite.Reset();
+	Result.CachedJson = ConvertRapidJsonToSharedJsonObject(Object);
+	Result.AdditionalFieldsToWrite.Reset();
 #endif // WITH_EDITOR
 
-	bool bSuccess = true;
-
 	// Get the name
-	if (!Object.TryGetStringField(TEXT("Name"), Name))
+	if (!TryGetStringField(Object, TEXT("Name"), Result.Name))
 	{
-		if (OutFailReason)
-		{
-			*OutFailReason = LOCTEXT("PluginReferenceWithoutName", "Plugin references must have a 'Name' field");
-		}
-		bSuccess = false;
+		return LOCTEXT("PluginReferenceWithoutName", "Plugin references must have a 'Name' field");
 	}
 
 	// Get the enabled field
-	if (!Object.TryGetBoolField(TEXT("Enabled"), bEnabled))
+	if (!TryGetBoolField(Object, TEXT("Enabled"), Result.bEnabled))
 	{
-		if (OutFailReason)
-		{
-			*OutFailReason = LOCTEXT("PluginReferenceWithoutEnabled", "Plugin references must have an 'Enabled' field");
-		}
-		bSuccess = false;
+		return LOCTEXT("PluginReferenceWithoutEnabled", "Plugin references must have an 'Enabled' field");
 	}
 
 	// Read the optional field
-	Object.TryGetBoolField(TEXT("Optional"), bOptional);
+	TryGetBoolField(Object, TEXT("Optional"), Result.bOptional);
 
-	Object.TryGetBoolField(TEXT("Activate"), bActivate);
+	TryGetBoolField(Object, TEXT("Activate"), Result.bActivate);
 
 	// Read the metadata for users that don't have the plugin installed
-	Object.TryGetStringField(TEXT("Description"), Description);
-	Object.TryGetStringField(TEXT("MarketplaceURL"), MarketplaceURL);
+	TryGetStringField(Object, TEXT("Description"), Result.Description);
+	TryGetStringField(Object, TEXT("MarketplaceURL"), Result.MarketplaceURL);
 
 	// Get the platform lists
-	JsonExtensions::TryGetStringArrayFieldWithDeprecatedFallback(Object, TEXT("PlatformAllowList"), TEXT("WhitelistPlatforms"), /*out*/ PlatformAllowList);
-	JsonExtensions::TryGetStringArrayFieldWithDeprecatedFallback(Object, TEXT("PlatformDenyList"), TEXT("BlacklistPlatforms"), /*out*/ PlatformDenyList);
+	TryGetStringArrayFieldWithDeprecatedFallback(Object, TEXT("PlatformAllowList"), TEXT("WhitelistPlatforms"), /*out*/ Result.PlatformAllowList);
+	TryGetStringArrayFieldWithDeprecatedFallback(Object, TEXT("PlatformDenyList"), TEXT("BlacklistPlatforms"), /*out*/ Result.PlatformDenyList);
 
 	// Get the target configuration lists
-	JsonExtensions::TryGetEnumArrayFieldWithDeprecatedFallback(Object, TEXT("TargetConfigurationAllowList"), TEXT("WhitelistTargetConfigurations"), /*out*/ TargetConfigurationAllowList);
-	JsonExtensions::TryGetEnumArrayFieldWithDeprecatedFallback(Object, TEXT("TargetConfigurationDenyList"), TEXT("BlacklistTargetConfigurations"), /*out*/ TargetConfigurationDenyList);
+	TryGetEnumArrayFieldWithDeprecatedFallback(Object, TEXT("TargetConfigurationAllowList"), TEXT("WhitelistTargetConfigurations"), /*out*/ Result.TargetConfigurationAllowList);
+	TryGetEnumArrayFieldWithDeprecatedFallback(Object, TEXT("TargetConfigurationDenyList"), TEXT("BlacklistTargetConfigurations"), /*out*/ Result.TargetConfigurationDenyList);
 
 	// Get the target lists
-	JsonExtensions::TryGetEnumArrayFieldWithDeprecatedFallback(Object, TEXT("TargetAllowList"), TEXT("WhitelistTargets"), /*out*/ TargetAllowList);
-	JsonExtensions::TryGetEnumArrayFieldWithDeprecatedFallback(Object, TEXT("TargetDenyList"), TEXT("BlacklistTargets"), /*out*/ TargetDenyList);
+	TryGetEnumArrayFieldWithDeprecatedFallback(Object, TEXT("TargetAllowList"), TEXT("WhitelistTargets"), /*out*/ Result.TargetAllowList);
+	TryGetEnumArrayFieldWithDeprecatedFallback(Object, TEXT("TargetDenyList"), TEXT("BlacklistTargets"), /*out*/ Result.TargetDenyList);
 
 	// Get the supported platform list
-	Object.TryGetStringArrayField(TEXT("SupportedTargetPlatforms"), SupportedTargetPlatforms);
-	Object.TryGetBoolField(TEXT("HasExplicitPlatforms"), bHasExplicitPlatforms);
+	TryGetStringArrayField(Object, TEXT("SupportedTargetPlatforms"), Result.SupportedTargetPlatforms);
+	TryGetBoolField(Object, TEXT("HasExplicitPlatforms"), Result.bHasExplicitPlatforms);
 
 	int32 ReadVersion;
-	if (Object.TryGetNumberField(TEXT("Version"), ReadVersion))
+	if (TryGetNumberField(Object, TEXT("Version"), ReadVersion))
 	{
-		RequestedVersion = ReadVersion;
-		if (!bEnabled)
+		Result.RequestedVersion = ReadVersion;
+		if (!Result.bEnabled)
 		{
-			if (bSuccess && OutFailReason)
-			{
-				*OutFailReason = LOCTEXT("PluginReferenceDisabledWithVersion", "Plugin references cannot be used to disable explicit versions. Remove the 'Version' field when 'Enabled' is false.");
+			return LOCTEXT("PluginReferenceDisabledWithVersion", "Plugin references cannot be used to disable explicit versions. Remove the 'Version' field when 'Enabled' is false.");
 			}
-			bSuccess = false;
 		}
-	}
 
-	return bSuccess;
+	return {};
 }
 
-bool FPluginReferenceDescriptor::Read(const FJsonObject& Object, FText& OutFailReason, TSharedPtr<FJsonObject> ObjectPtr /*= nullptr*/)
+bool FPluginReferenceDescriptor::Read(const TSharedRef<FJsonObject>& Object, FText* OutFailReason /*= nullptr*/)
 {
-PRAGMA_DISABLE_DEPRECATION_WARNINGS
-	return Read(Object, &OutFailReason, ObjectPtr);
-PRAGMA_ENABLE_DEPRECATION_WARNINGS
+	return UE::Projects::Private::ReadFromDefaultJson(*Object, *this, OutFailReason);
 }
 
-bool FPluginReferenceDescriptor::ReadArray(const FJsonObject& Object, const TCHAR* Name, TArray<FPluginReferenceDescriptor>& OutPlugins, FText* OutFailReason /*= nullptr*/)
+bool FPluginReferenceDescriptor::Read(const FJsonObject& Object, FText* OutFailReason /*= nullptr*/, TSharedPtr<FJsonObject> ObjectPtr /*= nullptr*/)
 {
-	const TArray< TSharedPtr<FJsonValue> >* Array;
-
-	if (Object.TryGetArrayField(Name, Array))
-	{
-		for (const TSharedPtr<FJsonValue> &Item : *Array)
-		{
-			const TSharedPtr<FJsonObject>* ItemObject = nullptr;
-
-			if (Item.IsValid() && Item->TryGetObject(ItemObject) && ItemObject && ItemObject->IsValid())
-			{
-				FPluginReferenceDescriptor PluginRef;
-
-				if (!PluginRef.Read(ItemObject->ToSharedRef(), OutFailReason))
-				{
-					return false;
-				}
+	return UE::Projects::Private::ReadFromDefaultJson(Object, *this, OutFailReason);
+}
 
-				OutPlugins.Add(PluginRef);
-			}
-		}
-	}
+bool FPluginReferenceDescriptor::Read(const FJsonObject& Object, FText& OutFailReason, TSharedPtr<FJsonObject> ObjectPtr /*= nullptr*/)
+{
+	return UE::Projects::Private::ReadFromDefaultJson(Object, *this, &OutFailReason);
+}
 
-	return true;
+bool FPluginReferenceDescriptor::ReadArray(const FJsonObject& Object, const TCHAR* Name, TArray<FPluginReferenceDescriptor>& OutPlugins, FText* OutFailReason /*= nullptr*/)
+{
+	return UE::Projects::Private::ReadArrayFromDefaultJson(Object, Name, OutPlugins, OutFailReason);
 }
 
 bool FPluginReferenceDescriptor::ReadArray(const FJsonObject& Object, const TCHAR* Name, TArray<FPluginReferenceDescriptor>& OutPlugins, FText& OutFailReason)
 {
-	return ReadArray(Object, Name, OutPlugins, &OutFailReason);
+	return UE::Projects::Private::ReadArrayFromDefaultJson(Object, Name, OutPlugins, &OutFailReason);
 }
 
 void FPluginReferenceDescriptor::Write(TJsonWriter<>& Writer) const
diff --git a/Engine/Source/Runtime/Projects/Private/ProjectDescriptor.cpp b/Engine/Source/Runtime/Projects/Private/ProjectDescriptor.cpp
--- a/Engine/Source/Runtime/Projects/Private/ProjectDescriptor.cpp
+++ b/Engine/Source/Runtime/Projects/Private/ProjectDescriptor.cpp
@@ -8,6 +8,8 @@
 #include "Serialization/JsonSerializer.h"
 #include "Interfaces/IProjectManager.h"
 #include "Misc/DataDrivenPlatformInfoRegistry.h"
+#include "JsonUtils/JsonConversion.h"
+#include "RapidJsonPluginLoading.h"
 
 #define LOCTEXT_NAMESPACE "ProjectDescriptor"
 
@@ -55,6 +57,8 @@ void FProjectDescriptor::UpdateSupportedTargetPlatforms(const FName& InPlatformN
 
 bool FProjectDescriptor::Load(const FString& FileName, FText& OutFailReason)
 {
+	SCOPED_BOOT_TIMING("FProjectDescriptor::Load")
+
 	// Read the file to a string
 	FString FileContents;
 	if (!FFileHelper::LoadFileToString(FileContents, *FileName))
@@ -65,28 +69,42 @@ bool FProjectDescriptor::Load(const FString& FileName, FText& OutFailReason)
 		return false;
 	}
 
-	// Deserialize a JSON object from the string
-	TSharedPtr< FJsonObject > Object;
-	TSharedRef< TJsonReader<> > Reader = TJsonReaderFactory<>::Create(FileContents);
-	if ( !FJsonSerializer::Deserialize(Reader, Object) || !Object.IsValid() )
+	// grab the characters so we can parse insitu safely (FString API says adding non-terminating zeros can be unsafe so just take the data)
+	TArray<TCHAR> Characters = MoveTemp(FileContents.GetCharArray());
+	TValueOrError<UE::Json::FDocument, UE::Json::FParseError> ParseResult = UE::Json::ParseInPlace(Characters);
+
+	if (ParseResult.HasError())
 	{
+		// have to re-load the file because it was destructively altered
+		ensure(FFileHelper::LoadFileToString(FileContents, *FileName));
+
 		// Do not try to localize these messages, a missing project file usually indicates missing data and an
 		// access to ICU will result in failure that's harder to diagnose
-		OutFailReason = FText::FromString(FString::Printf(TEXT("Failed to read file %s"), *FileName));
+		OutFailReason = FText::FromString(FString::Printf(TEXT("Failed to read file %s: %s"), *FileName, *ParseResult.GetError().CreateMessage(FileContents)));
+		return false;
+	}
+
+	TOptional<UE::Json::FConstObject> RootObject = UE::Json::GetRootObject(ParseResult.GetValue());
+	if (!RootObject.IsSet())
+	{
 		return false;
 	}
 
-	// Parse it as a project descriptor
-	return Read(*Object.Get(), FPaths::GetPath(FileName), OutFailReason);
+	return Read(*RootObject, FPaths::GetPath(FileName), OutFailReason);
 }
 
-bool FProjectDescriptor::Read(const FJsonObject& Object, const FString& PathToProject, FText& OutFailReason)
+bool FProjectDescriptor::Read(const UE::Json::FConstObject Object, const FString& PathToProject, FText& OutFailReason)
 {
+	SCOPED_BOOT_TIMING("FProjectDescriptor::Read")
+
+	using namespace UE::Json;
+	using namespace UE::Projects::Private;
+
 	// Read the file version
 	int32 FileVersionInt32;
-	if(!Object.TryGetNumberField(TEXT("FileVersion"), FileVersionInt32))
+	if(!TryGetNumberField(Object, TEXT("FileVersion"), FileVersionInt32))
 	{
-		if(!Object.TryGetNumberField(TEXT("ProjectFileVersion"), FileVersionInt32))
+		if(!TryGetNumberField(Object, TEXT("ProjectFileVersion"), FileVersionInt32))
 		{
 			OutFailReason = LOCTEXT("InvalidProjectFileVersion", "File does not have a valid 'FileVersion' number.");
 			return false;
@@ -104,48 +122,45 @@ bool FProjectDescriptor::Read(const FJsonObject& Object, const FString& PathToPr
 	}
 
 	// Read simple fields
-	Object.TryGetStringField(TEXT("EngineAssociation"), EngineAssociation);
-	Object.TryGetStringField(TEXT("Category"), Category);
-	Object.TryGetStringField(TEXT("Description"), Description);
-	Object.TryGetBoolField(TEXT("Enterprise"), bIsEnterpriseProject);
-	Object.TryGetBoolField(TEXT("DisableEnginePluginsByDefault"), bDisableEnginePluginsByDefault);
+	TryGetStringField(Object, TEXT("EngineAssociation"), EngineAssociation);
+	TryGetStringField(Object, TEXT("Category"), Category);
+	TryGetStringField(Object, TEXT("Description"), Description);
+	TryGetBoolField(Object, TEXT("Enterprise"), bIsEnterpriseProject);
+	TryGetBoolField(Object, TEXT("DisableEnginePluginsByDefault"), bDisableEnginePluginsByDefault);
 
 #if WITH_EDITOR
 	ModuleNamesCache.Reset();
 #endif
 
-	// Read the modules
-	if(!FModuleDescriptor::ReadArray(Object, TEXT("Modules"), Modules, OutFailReason))
+	if (TOptional<FText> ReadModulesError = ReadArray(Object, TEXT("Modules"), Modules))
 	{
+		OutFailReason = *ReadModulesError;
 		return false;
 	}
 
-	// Read the plugins
-	if(!FPluginReferenceDescriptor::ReadArray(Object, TEXT("Plugins"), Plugins, OutFailReason))
+	if (TOptional<FText> ReadPluginsError = ReadArray(Object, TEXT("Plugins"), Plugins))
 	{
+		OutFailReason = *ReadPluginsError;
 		return false;
 	}
 
 	// Read the list of additional plugin directories to scan
-	const TArray< TSharedPtr<FJsonValue> >* AdditionalPluginDirectoriesValue;
-	if (Object.TryGetArrayField(TEXT("AdditionalPluginDirectories"), AdditionalPluginDirectoriesValue))
+	TArray<FString> AdditionalPluginDirectoriesValue;
+	if (TryGetStringArrayField(Object, TEXT("AdditionalPluginDirectories"), AdditionalPluginDirectoriesValue))
 	{
 #if WITH_EDITOR || (IS_PROGRAM && WITH_PLUGIN_SUPPORT)
-		for (int32 Idx = 0; Idx < AdditionalPluginDirectoriesValue->Num(); Idx++)
+		for (FString& AdditionalPluginDir : AdditionalPluginDirectoriesValue)
 		{
-			FString AdditionalPluginDir;
-			if ((*AdditionalPluginDirectoriesValue)[Idx]->TryGetString(AdditionalPluginDir))
-			{
 				if (FPaths::IsRelative(AdditionalPluginDir))
 				{
 					AdditionalPluginDir = IFileManager::Get().ConvertToAbsolutePathForExternalAppForWrite(*(PathToProject / AdditionalPluginDir));
 				}
 				AddPluginDirectory(AdditionalPluginDir);
 			}
-		}
+
 #endif //if WITH_EDITOR
 		// If this is a packaged build and there are additional directories, they need to be remapped to the packaged location
-		if (FPlatformProperties::RequiresCookedData() && AdditionalPluginDirectoriesValue->Num() > 0)
+		if (FPlatformProperties::RequiresCookedData() && AdditionalPluginDirectoriesValue.Num() > 0)
 		{
 			AdditionalPluginDirectories.Empty();
 			FString RemappedDir = FPaths::ProjectDir() / TEXT("../RemappedPlugins/");
@@ -159,14 +174,11 @@ bool FProjectDescriptor::Read(const FJsonObject& Object, const FString& PathToPr
 
 #if WITH_EDITOR
 	// Read the list of additional root directories to scan
-	const TArray< TSharedPtr<FJsonValue> >* AdditionalRootDirectoriesValue;
-	if (Object.TryGetArrayField(TEXT("AdditionalRootDirectories"), AdditionalRootDirectoriesValue))
+	TArray<FString> AdditionalRootDirectoriesValue;
+	if (TryGetStringArrayField(Object, TEXT("AdditionalRootDirectories"), AdditionalRootDirectoriesValue))
 	{
-		for (const TSharedPtr<FJsonValue>& AdditionalRootDirectoryValue : *AdditionalRootDirectoriesValue)
+		for (FString& AdditionalRoot : AdditionalRootDirectoriesValue)
 		{
-			FString AdditionalRoot;
-			if (AdditionalRootDirectoryValue->TryGetString(AdditionalRoot))
-			{
 				if (FPaths::IsRelative(AdditionalRoot))
 				{
 					AdditionalRoot = IFileManager::Get().ConvertToAbsolutePathForExternalAppForWrite(*(PathToProject / AdditionalRoot));
@@ -174,22 +186,10 @@ bool FProjectDescriptor::Read(const FJsonObject& Object, const FString& PathToPr
 				AddRootDirectory(AdditionalRoot);
 			}
 		}
-	}
 #endif //if WITH_EDITOR
 
 	// Read the target platforms
-	const TArray< TSharedPtr<FJsonValue> > *TargetPlatformsValue;
-	if(Object.TryGetArrayField(TEXT("TargetPlatforms"), TargetPlatformsValue))
-	{
-		for(int32 Idx = 0; Idx < TargetPlatformsValue->Num(); Idx++)
-		{
-			FString TargetPlatform;
-			if((*TargetPlatformsValue)[Idx]->TryGetString(TargetPlatform))
-			{
-				TargetPlatforms.Add(*TargetPlatform);
-			}
-		}
-	}
+	TryGetStringArrayField(Object, TEXT("TargetPlatforms"), TargetPlatforms);
 
 	// check if the project has directories for extended platforms, and assume support if it does
 	// however, if there were no platforms already listed, then all platforms are supported, and we don't
@@ -220,15 +220,26 @@ bool FProjectDescriptor::Read(const FJsonObject& Object, const FString& PathToPr
 	}
 
 	// Get the sample name hash
-	Object.TryGetNumberField(TEXT("EpicSampleNameHash"), EpicSampleNameHash);
+	TryGetNumberField(Object, TEXT("EpicSampleNameHash"), EpicSampleNameHash);
 
 	// Read the custom build steps
-	PreBuildSteps.Read(Object, TEXT("PreBuildSteps"));
-	PostBuildSteps.Read(Object, TEXT("PostBuildSteps"));
+	PreBuildSteps = ReadCustomBuildSteps(Object, TEXT("PreBuildSteps"));
+	PostBuildSteps = ReadCustomBuildSteps(Object, TEXT("PostBuildSteps"));
 
 	return true;
 }
 
+bool FProjectDescriptor::Read(const FJsonObject& Object, const FString& PathToProject, FText& OutFailReason)
+{
+	return UE::Projects::Private::ReadFromDefaultJsonHelper(
+		Object,
+		[&PathToProject, &OutFailReason, this](UE::Json::FConstObject Object)
+		{
+			return this->Read(Object, PathToProject, OutFailReason);
+		}
+	);
+}
+
 bool FProjectDescriptor::Save(const FString& FileName, FText& OutFailReason)
 {
 	if (IProjectManager::Get().IsSuppressingProjectFileWrite())
diff --git a/Engine/Source/Runtime/Projects/Private/RapidJsonPluginLoading.cpp b/Engine/Source/Runtime/Projects/Private/RapidJsonPluginLoading.cpp
new file mode 100644
--- /dev/null
+++ b/Engine/Source/Runtime/Projects/Private/RapidJsonPluginLoading.cpp
@@ -0,0 +1,301 @@
+// Copyright Epic Games, Inc. All Rights Reserved.
+
+#include "RapidJsonPluginLoading.h"
+#include "Internationalization/Internationalization.h"
+
+#define LOCTEXT_NAMESPACE "RapidJsonPluginLoading"
+
+DEFINE_LOG_CATEGORY_STATIC(LogRapidJsonPluginLoading, Display, All);
+
+namespace UE {
+namespace Projects {
+namespace Private {
+
+bool TryGetBoolField(Json::FConstObject Object, const TCHAR* FieldName, bool& Out)
+{
+	Json::FValue::ConstMemberIterator It = Object.FindMember(FieldName);
+	if (It == Object.MemberEnd())
+	{
+		return false;
+	}
+
+	if (It->value.IsBool())
+	{
+		Out = It->value.GetBool();
+		return true;
+	}
+	else if (It->value.IsInt64())
+	{
+		Out = (It->value.GetInt64() != 0);
+		return true;
+	}
+	else if (It->value.IsUint64())
+	{
+		Out = (It->value.GetUint64() != 0);
+		return true;
+	}
+	else if (It->value.IsDouble())
+	{
+		Out = (It->value.GetDouble() != 0.0);
+		return true;
+	}
+	else if (It->value.IsString())
+	{
+		Out = TCString<TCHAR>::ToBool(It->value.GetString());
+		return true;
+	}
+
+	return false;
+}
+
+// this 'TryConvertFromDouble' code was copy/pasted from the FJsonValue code, to maintain compatibility. When the data converted from the shared json for a 
+// legacy code path it loses any of the original type meaning. So parsing from this double version needs to behave the same as if we were reading form integers
+//
+// As soon as the public FJsonObject read API can be deprecated this can be removed
+template <typename T>
+bool TryConvertFromDouble(double Double, T& OutNumber)
+{
+	if ((Double >= TNumericLimits<T>::Min()) && (Double <= static_cast<double>(TNumericLimits<T>::Max())))
+	{
+		OutNumber = static_cast<T>(FMath::RoundHalfFromZero(Double));
+
+		return true;
+	}
+
+	return false;
+}
+
+// Need special handling for int64/uint64, due to overflow in the numeric limits.
+// 2^63-1 and 2^64-1 cannot be exactly represented as a double, so TNumericLimits<>::Max() gets rounded up to exactly 2^63 or 2^64 by the compiler's implicit cast to double.
+// This breaks the overflow check in TryConvertNumber. We use "<" rather than "<=" along with the exact power-of-two double literal to fix this.
+template <> 
+bool TryConvertFromDouble<uint64>(double Double, uint64& OutNumber)
+{
+	if (Double >= 0.0 && Double < 18446744073709551616.0)
+	{
+		OutNumber = static_cast<uint64>(FMath::RoundHalfFromZero(Double));
+		return true;
+	}
+
+	return false;
+}
+
+template <>
+bool TryConvertFromDouble<int64>(double Double, int64& OutNumber)
+{
+	if (Double >= -9223372036854775808.0 && Double < 9223372036854775808.0)
+	{
+		OutNumber = static_cast<int64>(FMath::RoundHalfFromZero(Double));
+		return true;
+	}
+
+	return false;
+}
+
+// returns true if the value can be safely stored in output
+template<typename T>
+static inline bool TryGetIntHelper(Json::FConstObject Object, const TCHAR* FieldName, const TCHAR* OutTypeName, T& Out)
+{
+	using namespace Json;
+
+	FValue::ConstMemberIterator It = Object.FindMember(FieldName);
+	if (It == Object.MemberEnd())
+	{
+		return false;
+	}
+
+	if (It->value.IsInt64())
+	{
+		const int64 Value = It->value.GetInt64();
+		if (IntFitsIn<T>(Value))
+		{
+			Out = static_cast<T>(Value);
+			return true;
+		}
+		else
+		{
+			return false;
+		}
+	}
+	else if (It->value.IsUint64())
+	{
+		const uint64 Value = It->value.GetUint64();
+		if (IntFitsIn<T>(Value))
+		{
+			Out = static_cast<T>(Value);
+			return true;
+		}
+		else
+		{
+			return false;
+		}
+	}
+	else if (It->value.IsDouble())
+	{
+		const double Value = It->value.GetDouble();
+		if (TryConvertFromDouble<T>(Value, Out))
+		{
+			return true;
+		}
+		else
+		{
+			return false;
+		}
+	}
+	else if (It->value.IsBool())
+	{
+		Out = It->value.GetBool() ? 1 : 0;
+		return true;
+	}
+	else if (It->value.IsString())
+	{
+		LexFromString(Out, It->value.GetString());
+		return true;
+	}
+	else
+	{
+		return false;
+	}
+}
+
+bool TryGetNumberField(Json::FConstObject Object, const TCHAR* FieldName, int32& Out)
+{
+	return TryGetIntHelper(Object, FieldName, TEXT("int32"), Out);
+}
+
+bool TryGetNumberField(Json::FConstObject Object, const TCHAR* FieldName, uint32& Out)
+{
+	return TryGetIntHelper(Object, FieldName, TEXT("uint32"), Out);
+}
+
+static inline bool ConvertToStringHelper(const Json::FValue& Value, FString& Out)
+{
+	if (Value.IsString())
+	{
+		Out = Value.GetString();
+		return true;
+	}
+	else if (Value.IsInt64())
+	{
+		Out = FString::SanitizeFloat(static_cast<double>(Value.GetInt64()), 0);
+		return true;
+	}
+	else if (Value.IsDouble())
+	{
+		Out = FString::SanitizeFloat(Value.GetDouble(), 0);
+		return true;
+	}
+	else if (Value.IsBool())
+	{
+		Out = Value.GetBool() ? TEXT("true") : TEXT("false");
+		return true;
+	}
+	else
+	{
+		Out.Reset();
+		return false;
+	}
+}
+
+bool TryGetStringField(Json::FConstObject Object, const TCHAR* FieldName, FString& Out)
+{
+	using namespace Json;
+
+	FValue::ConstMemberIterator It = Object.FindMember(FieldName);
+	if (It == Object.MemberEnd())
+	{
+		return false;
+	}
+
+	if (ConvertToStringHelper(It->value, Out))
+	{
+		return true;
+	}
+
+	return false;
+}
+
+bool TryGetStringArrayField(Json::FConstObject Object, const TCHAR* FieldName, TArray<FString>& Out)
+{
+	using namespace Json;
+
+	Out.Reset();
+
+	FValue::ConstMemberIterator It = Object.FindMember(FieldName);
+	if (It == Object.MemberEnd() || !It->value.IsArray())
+	{
+		return false;
+	}
+
+	FConstArray Array = It->value.GetArray();
+	Out.Reserve(Array.Size());
+	for (const FValue& Item : Array)
+	{
+		if (!ConvertToStringHelper(Item, Out.Emplace_GetRef()))
+		{
+			Out.Pop();
+			return false;
+		}
+	}
+
+	return true;
+}
+
+bool TryGetStringArrayField(Json::FConstObject Object, const TCHAR* FieldName, TArray<FName>& Out)
+{
+	using namespace Json;
+
+	Out.Reset();
+
+	FValue::ConstMemberIterator It = Object.FindMember(FieldName);
+	if (It == Object.MemberEnd() || !It->value.IsArray())
+	{
+		return false;
+	}
+
+	FConstArray Array = It->value.GetArray();
+	Out.Reserve(Array.Size());
+
+	FString TempString;
+
+	for (const FValue& Item : Array)
+	{
+		if (ConvertToStringHelper(Item, TempString))
+		{
+			Out.Emplace(TempString);
+		}
+		else
+		{
+			return false;
+		}
+	}
+
+	return true;
+}
+
+bool TryGetStringArrayFieldWithDeprecatedFallback(Json::FConstObject Object, const TCHAR* FieldName, const TCHAR* DeprecatedFieldName, TArray<FString>& OutArray)
+{
+	if (TryGetStringArrayField(Object, FieldName, OutArray))
+	{
+		return true;
+	}
+
+	return TryGetStringArrayField(Object, DeprecatedFieldName, OutArray);
+}	
+
+FText GetArrayObjectTypeError(const TCHAR* FieldName, int32 Index)
+{
+	return FText::Format(LOCTEXT("InvalidPluginArrayEntry", "Field '{0}[{1}]' is a non-object entry"), FText::FromString(FieldName), Index);
+}
+
+
+FText GetArrayObjectChildParseError(const TCHAR* FieldName, int32 Index, const FText& PropagateError)
+{
+	return FText::Format(LOCTEXT("InvalidPluginParseArrayEntry", "'{0}[{1}]': {2}"), FText::FromString(FieldName), Index, PropagateError);
+}
+
+} // Private
+} // Projects
+} // UE
+
+#undef LOCTEXT_NAMESPACE
\ No newline at end of file
diff --git a/Engine/Source/Runtime/Projects/Private/RapidJsonPluginLoading.h b/Engine/Source/Runtime/Projects/Private/RapidJsonPluginLoading.h
new file mode 100644
--- /dev/null
+++ b/Engine/Source/Runtime/Projects/Private/RapidJsonPluginLoading.h
@@ -0,0 +1,192 @@
+// Copyright Epic Games, Inc. All Rights Reserved.
+
+#pragma once
+
+#include "JsonUtils/RapidJsonUtils.h"
+#include "JsonUtils/JsonConversion.h"
+#include "Internationalization/Text.h"
+
+struct FLocalizationTargetDescriptor;
+struct FModuleDescriptor;
+struct FPluginDescriptor;
+struct FCustomBuildSteps;
+struct FPluginReferenceDescriptor;
+struct FPluginDisallowedDescriptor;
+struct FPluginManifest;
+
+namespace UE {
+namespace Projects {
+namespace Private {
+
+/**
+ * These are utility functions to be as close to the default Json parser as possible, only really necessary
+ * because we have public APIs that take FJsonObjects directly
+*/
+
+bool TryGetBoolField(Json::FConstObject Object, const TCHAR* FieldName, bool& Out);
+bool TryGetNumberField(Json::FConstObject Object, const TCHAR* FieldName, int32& Out);
+bool TryGetNumberField(Json::FConstObject Object, const TCHAR* FieldName, uint32& Out);
+bool TryGetStringField(Json::FConstObject Object, const TCHAR* FieldName, FString& Out);
+bool TryGetStringField(Json::FConstObject Object, const TCHAR* FieldName, FStringView& Out);
+bool TryGetStringArrayField(Json::FConstObject Object, const TCHAR* FieldName, TArray<FString>& Out);
+bool TryGetStringArrayField(Json::FConstObject Object, const TCHAR* FieldName, TArray<FName>& Out);	
+
+bool TryGetStringArrayFieldWithDeprecatedFallback(Json::FConstObject Object, const TCHAR* FieldName, const TCHAR* DeprecatedFieldName, TArray<FString>& OutArray);
+
+FText GetArrayObjectTypeError(const TCHAR* FieldName, int32 Index);
+FText GetArrayObjectChildParseError(const TCHAR* FieldName, int32 Index, const FText& PropagateError);
+
+
+/** Get the field named FieldName as an array of strings. Returns false if it doesn't exist or any member cannot be converted. */
+template<typename TEnum>
+inline bool TryGetEnumArrayField(Json::FConstObject Object, const TCHAR* FieldName, TArray<TEnum>& OutArray)
+{
+	TOptional<Json::FConstArray> Array = Json::GetArrayField(Object, FieldName);
+	if (!Array.IsSet())
+	{
+		return false;
+	}
+
+	OutArray.Reset(Array->Size());
+	for (const Json::FValue& JsonValue : *Array)
+	{
+		if (!JsonValue.IsString())
+		{
+			continue;
+		}
+
+		TEnum Value;
+		if (LexTryParseString(Value, JsonValue.GetString()))
+		{
+			OutArray.Add(Value);
+		}
+	}
+	return true;
+}
+
+template<typename TEnum>
+inline bool TryGetEnumArrayFieldWithDeprecatedFallback(Json::FConstObject Object, const TCHAR* FieldName, const TCHAR* DeprecatedFieldName, TArray<TEnum>& OutArray)
+{
+	if (TryGetEnumArrayField<TEnum>(Object, FieldName, /*out*/ OutArray))
+	{
+		return true;
+	}
+	else if (TryGetEnumArrayField<TEnum>(Object, DeprecatedFieldName, /*out*/ OutArray))
+	{
+		//@TODO: Warn about deprecated field fallback?
+		return true;
+	}
+	else
+	{
+		return false;
+	}
+}
+
+// A private interface as to not expose RapidJson types
+
+TOptional<FText> Read(Json::FConstObject Object, FPluginManifest& Out);
+TOptional<FText> Read(Json::FConstObject Object, FPluginDescriptor& Out);
+TOptional<FText> Read(Json::FConstObject Object, FModuleDescriptor& Out);
+TOptional<FText> Read(Json::FConstObject Object, FLocalizationTargetDescriptor& Out);
+TOptional<FText> Read(Json::FConstObject Object, FPluginReferenceDescriptor& Out);
+TOptional<FText> Read(Json::FConstObject Object, FPluginDisallowedDescriptor& Out);
+
+FCustomBuildSteps ReadCustomBuildSteps(Json::FConstObject Object, const TCHAR* FieldName);
+
+// Helper functions to convert and call the rapidjson functions from default json sources
+template<typename FUNCTOR>
+bool ReadFromDefaultJsonHelper(const FJsonObject& Object, FUNCTOR && Functor)
+{
+	TOptional<Json::FDocument> Document = Json::ConvertSharedJsonToRapidJsonDocument(Object);
+	if (!Document.IsSet())
+	{
+		return false;
+	}
+
+	TOptional<Json::FConstObject> RootObject = Json::GetRootObject(*Document);
+	if (!RootObject.IsSet())
+	{
+		return false;
+	}
+
+	return Functor(*RootObject);
+}
+
+template<typename TYPE>
+bool ReadFromDefaultJson(const FJsonObject& Object, TYPE& Out, FText* OutFailReason)
+{
+	return ReadFromDefaultJsonHelper(
+		Object,
+		[&Out, OutFailReason](Json::FConstObject RootObject)
+		{
+			TOptional<FText> ResultError = Read(RootObject, Out);
+			if (ResultError)
+			{
+				if (OutFailReason)
+				{
+					*OutFailReason = *ResultError;
+				}
+
+				return false;
+			}
+
+			return true;
+		}
+	);
+}	
+
+template<typename TYPE>
+TOptional<FText> ReadArray(Json::FConstObject Object, const TCHAR* FieldName, TArray<TYPE>& OutArray)
+{
+	TOptional<Json::FConstArray> ArrayValue = Json::GetArrayField(Object, FieldName);
+	if(ArrayValue.IsSet())
+	{
+		OutArray.Reserve(ArrayValue->Size());
+
+		for (const Json::FValue& ArrayItem : *ArrayValue)
+		{
+			if (ArrayItem.IsObject())
+			{
+				TOptional<FText> EntryResult = Read(ArrayItem.GetObject(), OutArray.Emplace_GetRef());
+				if (EntryResult.IsSet())
+				{
+					OutArray.Pop();
+					return GetArrayObjectChildParseError(FieldName, OutArray.Num(), *EntryResult);
+				}
+			}
+			else
+			{
+				return GetArrayObjectTypeError(FieldName, OutArray.Num());
+			}
+		}
+	}
+	
+	return {};
+}
+
+template<typename TYPE>
+bool ReadArrayFromDefaultJson(const FJsonObject& Object, const TCHAR* Name, TArray<TYPE>& OutArray, FText* OutFailReason)
+{
+	return ReadFromDefaultJsonHelper(
+		Object,
+		[Name, &OutArray, OutFailReason](Json::FConstObject RootObject)
+		{
+			TOptional<FText> ArrayError = ReadArray<TYPE>(RootObject, Name, OutArray);
+			if (ArrayError)
+			{
+				if (OutFailReason)
+				{
+					*OutFailReason = MoveTemp(*ArrayError);
+				}
+
+				return false;
+			}
+
+			return true;
+		}
+	);
+}	
+	
+}
+}  
+}
\ No newline at end of file
diff --git a/Engine/Source/Runtime/Projects/Private/Tests/RapidJsonPluginLoadingTests.cpp b/Engine/Source/Runtime/Projects/Private/Tests/RapidJsonPluginLoadingTests.cpp
new file mode 100644
--- /dev/null
+++ b/Engine/Source/Runtime/Projects/Private/Tests/RapidJsonPluginLoadingTests.cpp
@@ -0,0 +1,449 @@
+// Copyright Epic Games, Inc. All Rights Reserved.
+
+#if WITH_TESTS
+
+#include "CoreMinimal.h"
+#include "Misc/AutomationTest.h"
+#include "Serialization/JsonSerializer.h"
+#include "Tests/TestHarnessAdapter.h"
+#include "RapidJsonPluginLoading.h"
+#include "PluginDescriptor.h"
+#include "ProjectDescriptor.h"
+
+#if WITH_LOW_LEVEL_TESTS
+#include "TestCommon/Expectations.h"
+#endif
+
+/** Ensure RapidJsonUtils parses with the same quirks as the default Json loader, e.g. converting between types correctly **/
+TEST_CASE_NAMED(FProjectsLoadingRapidJsonPluginLoading, "System::Engine::Projects::Loading::RapidJsonPluginLoading", "[ApplicationContextMask][SmokeFilter]")
+{
+	const TCHAR JsonText[] = TEXT(R"json(
+{
+	"zero": 0,
+	"min_int32": -2147483648,
+	"max_int32":  2147483647,
+	"max_uint32": 4294967295,
+	"round_down": 1.1,
+	"round_up" : 1.9,
+	"negative_round_down": -1.9,
+	"negative_round_up" : -1.1,
+	"string": "hello world",
+	"on": true,
+	"off": false,
+	"object": {
+		"value": 42
+	},
+	"array": [
+		"aaa",
+		"bbb",
+		"ccc",
+		"ddd"
+	],	
+	"null": null
+}
+
+)json");
+
+	// check to make sure all values are parsed the same across RapidJsonUtils an default Json
+
+	TSharedRef< TJsonReader<> > Reader = TJsonReaderFactory<>::Create(JsonText);
+	TSharedPtr< FJsonObject > DefaultRoot;
+	REQUIRE(FJsonSerializer::Deserialize(Reader, DefaultRoot));
+
+	using namespace UE::Json;
+	using namespace UE::Projects::Private;
+
+	// convert to RapidJson
+	TValueOrError<FDocument, FParseError> RapidJsonDocument = UE::Json::Parse(JsonText);
+	FString ParseErrorMessage = RapidJsonDocument.HasError() ? RapidJsonDocument.GetError().CreateMessage(JsonText) : TEXT("No Error");
+	REQUIRE_MESSAGE(ParseErrorMessage, RapidJsonDocument.HasValue());
+
+	TOptional<UE::Json::FConstObject> RapidRoot = UE::Json::GetRootObject(RapidJsonDocument.GetValue());
+	REQUIRE(RapidRoot.IsSet());
+
+	auto CheckField = [&Default = DefaultRoot, Rapid = *RapidRoot](const TCHAR* FieldName)
+	{
+		int32 DefaultInt32{};
+		int32 RapidInt32{};
+		const bool bDefaultInt32Found = Default->TryGetNumberField(FieldName, DefaultInt32);
+		const bool bRapidInt32Found = TryGetNumberField(Rapid, FieldName, RapidInt32);
+		REQUIRE_MESSAGE(*FString::Printf(TEXT("%s: int32 Parse Mismatch: %d vs %d"), FieldName, int32(bDefaultInt32Found), int32(bRapidInt32Found)), bDefaultInt32Found == bRapidInt32Found);
+		REQUIRE_MESSAGE(*FString::Printf(TEXT("%s: int32 Mismatch: %d vs %d"), FieldName, DefaultInt32, RapidInt32), DefaultInt32 == RapidInt32);
+
+		uint32 DefaultUint32{};
+		uint32 RapidUint32{};
+		const bool bDefaultUint32Found = Default->TryGetNumberField(FieldName, DefaultUint32);
+		const bool bRapidUint32Found = TryGetNumberField(Rapid, FieldName, RapidUint32);
+		REQUIRE_MESSAGE(*FString::Printf(TEXT("%s: uint32 Parse Mismatch: %d vs %d"), FieldName, int32(bDefaultUint32Found), int32(bRapidUint32Found)), bDefaultUint32Found == bRapidUint32Found);
+		REQUIRE_MESSAGE(*FString::Printf(TEXT("%s: uint32 Mismatch: %u vs %u"), FieldName, DefaultUint32, RapidUint32), DefaultUint32 == RapidUint32);
+
+		bool DefaultBool{};
+		bool RapidBool{};
+		const bool bDefaultBoolFound = Default->TryGetBoolField(FieldName, DefaultBool);
+		const bool bRapidBoolFound = TryGetBoolField(Rapid, FieldName, RapidBool);
+		REQUIRE_MESSAGE(*FString::Printf(TEXT("%s: Bool Parse Mismatch: %d vs %d"), FieldName, int32(bDefaultBoolFound), int32(bRapidBoolFound)), bDefaultBoolFound == bRapidBoolFound);
+		REQUIRE_MESSAGE(TEXT("Bool Mismatch"), DefaultBool == RapidBool);
+
+		FString DefaultString{};
+		FString RapidString{};
+		const bool bDefaultStringFound = Default->TryGetStringField(FieldName, DefaultString);
+		const bool bRapidStringFound = TryGetStringField(Rapid, FieldName, RapidString);
+		REQUIRE_MESSAGE(*FString::Printf(TEXT("%s: String Parse Mismatch: %d vs %d"), FieldName, int32(bDefaultStringFound), int32(bRapidStringFound)), bDefaultStringFound == bRapidStringFound);
+		REQUIRE_MESSAGE(TEXT("String Mismatch"), DefaultString == RapidString);
+
+		const TArray<TSharedPtr<FJsonValue>>* DefaultArray{};
+		TOptional<FConstArray> RapidArray = GetArrayField(Rapid, FieldName);
+		REQUIRE_MESSAGE(*FString::Printf(TEXT("%s: Array Parse Mismatch"), FieldName), Default->TryGetArrayField(FieldName, DefaultArray) == RapidArray.IsSet());
+
+		if (DefaultArray && RapidArray.IsSet())
+		{
+			REQUIRE_MESSAGE(TEXT("Array Size Mismatch"), DefaultArray->Num() == RapidArray->Size());
+		}
+
+		const TSharedPtr<FJsonObject>* DefaultChildObject{};
+		TOptional<FConstObject> RapidChildObject = GetObjectField(Rapid, FieldName);
+		REQUIRE_MESSAGE(*FString::Printf(TEXT("%s: Object Parse Mismatch"), FieldName), Default->TryGetObjectField(FieldName, DefaultChildObject) == RapidChildObject.IsSet());
+
+		TSharedPtr<FJsonValue> DefaultNull = Default->TryGetField(FieldName);
+		const bool HasDefaultNull = DefaultNull ? DefaultNull->Type == EJson::Null : false;
+		const bool HasRapidNull = HasNullField(Rapid, FieldName);
+
+		REQUIRE_MESSAGE(*FString::Printf(TEXT("%s: Null Type Parse Mismatch"), FieldName), HasDefaultNull == HasRapidNull);
+	};
+
+	CheckField(TEXT("zero"));
+	CheckField(TEXT("min_int32"));
+	CheckField(TEXT("max_int32"));
+	CheckField(TEXT("max_uint32"));
+	CheckField(TEXT("round_down"));
+	CheckField(TEXT("round_up"));
+	CheckField(TEXT("negative_round_down"));
+	CheckField(TEXT("negative_round_up"));
+	CheckField(TEXT("string"));
+	CheckField(TEXT("on"));
+	CheckField(TEXT("off"));
+	CheckField(TEXT("object"));
+	CheckField(TEXT("array"));
+	CheckField(TEXT("null"));
+}
+
+
+// the name of the test case seems scrambled sometimes, so make a consistent helper struct that we can friend to test this function
+struct FProjectsLoadingPluginDescriptorTestsHelper
+{
+	static bool Read(FProjectDescriptor& ProjectDescriptor, const UE::Json::FConstObject Object, const FString& PathToProject, FText& OutFailReason)
+	{
+		return ProjectDescriptor.Read(Object, PathToProject, OutFailReason);
+	}
+};
+
+// make sure everything in the plugin descriptor loads
+TEST_CASE_NAMED(FProjectsLoadingPluginDescriptorTests, "System::Engine::Projects::Loading::PluginDescriptor", "[ApplicationContextMask][SmokeFilter]")
+{
+	// every property is here and set to a non-default value to ensure it reads
+	const TCHAR JsonText[] = TEXT(R"json(
+{
+	"FileVersion": 3,
+	"Version": 8,
+	"VersionName": "0.2",
+	"FriendlyName": "Test Plugin",
+	"Description": "Test Plugin for unit testing",
+	"Category": "Editor",
+	"CreatedBy": "Epic Games, Inc.",
+	"CreatedByURL": "https://epicgames.com",
+	"DocsURL": "DocsURL goes here",
+	"MarketplaceURL": "MarketplaceURL goes here",
+	"SupportURL": "SupportURL goes here",
+	"EngineVersion": "5.5",
+	"EditorCustomVirtualPath": "EditorCustomVirtualPath goes here",
+	"SupportedTargetPlatforms": ["Win64", "Mac"],
+	"SupportedPrograms": ["UnrealInsights"],
+	"bIsPluginExtension": true,
+	"EnabledByDefault": false,
+	"Installed": true,
+	"VersePath": "/VersePath.com",
+	"DeprecatedEngineVersion": "5.5",
+	"VerseScope": "InternalAPI",
+	"VerseVersion": 1,
+	"EnableSceneGraph": true,
+	"EnableVerseAssetReflection": true,
+	"EnableIAD": true,
+	"CanContainContent": true,
+	"CanContainVerse": true,
+	"NoCode": true,
+	"IsBetaVersion": true,
+	"IsExperimentalVersion": true,
+	"RequiresBuildPlatform": true,
+	"Hidden": true,
+	"Sealed": true,
+	"ExplicitlyLoaded": true,
+	"HasExplicitPlatforms": true,
+	"CanBeUsedWithUnrealHeaderTool": true,
+	"LocalizationTargets": [
+		{
+			"Name": "LocalizationTarget AAA",
+			"LoadingPolicy": "Always",
+			"ConfigGenerationPolicy": "Auto"
+		}
+	],
+	"Modules": [
+		{
+			"Name": "AAA",
+			"Type": "DeveloperTool",
+			"LoadingPhase": "Default",
+			"HasExplicitPlatforms": true,
+			"PlatformAllowList": ["Win64"],
+			"PlatformDenyList": ["Mac"],
+			"TargetAllowList" : ["Editor"],
+			"TargetDenyList" : ["Client"],
+			"TargetConfigurationAllowList" : ["Test"],
+			"TargetConfigurationDenyList" : ["DebugGame"],
+			"ProgramAllowList" : ["UnrealInsights"],
+			"ProgramDenyList" : ["LiveLinkHub"],
+			"GameTargetAllowList" : ["AAA"],
+			"GameTargetDenyList" : ["BBB"],
+			"AdditionalDependencies" : ["AAA", "BBB"]
+		},
+		{
+			"Name": "BBB",
+			"Type": "Editor",
+			"LoadingPhase": "Default"
+		},
+		{
+			"Name": "CCC",
+			"Type": "Runtime",
+			"LoadingPhase": "PostConfigInit"
+		}
+	],
+	"Plugins": [
+		{
+			"Name": "DDD",
+			"Enabled": true,
+			"Optional": true,
+			"Description": "Ref Description goes here",
+			"MarketplaceURL": "Ref MarketplaceURL goes here",
+			"PlatformAllowList": ["Win64"],
+			"PlatformDenyList": ["Mac"],
+			"TargetConfigurationAllowList" : ["Test"],
+			"TargetConfigurationDenyList" : ["DebugGame"],
+			"TargetAllowList" : ["Editor"],
+			"TargetDenyList" : ["Client"],
+			"SupportedTargetPlatforms": ["Win64", "Mac"],
+			"HasExplicitPlatforms": true,
+			"Version": 2
+		},
+		{
+			"Name": "EEE",
+			"Enabled": true
+		}
+	],
+	"PreBuildSteps": {
+		"PreBuildStepA": [
+			"Step 1a",
+			"Step 2a",
+		]
+	},
+	"PostBuildSteps": {
+		"PostBuildStepA": ["Step 1b"],
+		"PostBuildStepB": ["Step 2b"],
+	},
+	"DisallowedPlugins": [
+		{
+			"Name": "ABC",
+			"Comment": "ABC Reasons"
+		},
+	]
+}
+)json");
+
+	// check to make sure all values are parsed the same across RapidJsonUtils an default Json
+
+	FText FailReason;
+	FPluginDescriptor Descriptor;
+	REQUIRE_MESSAGE(FailReason.ToString(), Descriptor.Read(JsonText, FailReason));
+
+	// check properties were read
+	REQUIRE(Descriptor.Version == 8);
+	REQUIRE(Descriptor.VersionName == TEXT("0.2"));
+	REQUIRE(Descriptor.FriendlyName == TEXT("Test Plugin"));
+	REQUIRE(Descriptor.Description == TEXT("Test Plugin for unit testing"));
+	REQUIRE(Descriptor.Category == TEXT("Editor"));
+	REQUIRE(Descriptor.CreatedBy == TEXT("Epic Games, Inc."));
+	REQUIRE(Descriptor.CreatedByURL == TEXT("https://epicgames.com"));
+	REQUIRE(Descriptor.DocsURL == TEXT("DocsURL goes here"));
+	REQUIRE(Descriptor.MarketplaceURL == TEXT("MarketplaceURL goes here"));
+	REQUIRE(Descriptor.SupportURL == TEXT("SupportURL goes here"));
+	REQUIRE(Descriptor.EngineVersion == TEXT("5.5"));
+	REQUIRE(Descriptor.EditorCustomVirtualPath == TEXT("EditorCustomVirtualPath goes here"));
+	REQUIRE(Descriptor.SupportedTargetPlatforms.Num() == 2);
+	REQUIRE(Descriptor.SupportedTargetPlatforms[0] == TEXT("Win64"));
+	REQUIRE(Descriptor.SupportedTargetPlatforms[1] == TEXT("Mac"));
+	REQUIRE(Descriptor.SupportedPrograms.Num() == 2);
+	REQUIRE(Descriptor.SupportedPrograms[0] == TEXT("UnrealInsights"));
+	REQUIRE(Descriptor.SupportedPrograms[1] == TEXT("UnrealHeaderTool")); // added by CanBeUsedWithUnrealHeaderTool
+	REQUIRE(Descriptor.bIsPluginExtension == true);
+	REQUIRE(Descriptor.VersePath == TEXT("/VersePath.com"));
+	REQUIRE(Descriptor.DeprecatedEngineVersion == TEXT("5.5"));
+	REQUIRE(Descriptor.VerseScope == EVerseScope::InternalAPI);
+	REQUIRE(Descriptor.VerseVersion == 1);
+	REQUIRE(Descriptor.bEnableSceneGraph == true);
+	REQUIRE(Descriptor.bEnableVerseAssetReflection == true);
+	REQUIRE(Descriptor.bEnableIAD == true);
+	REQUIRE(Descriptor.EnabledByDefault == EPluginEnabledByDefault::Disabled);
+	REQUIRE(Descriptor.bCanContainContent == true);
+	REQUIRE(Descriptor.bCanContainVerse == true);
+	REQUIRE(Descriptor.bNoCode == true);
+	REQUIRE(Descriptor.bIsBetaVersion == true);
+	REQUIRE(Descriptor.bIsExperimentalVersion == true);
+	REQUIRE(Descriptor.bInstalled == true);
+	REQUIRE(Descriptor.bRequiresBuildPlatform == true);
+	REQUIRE(Descriptor.bIsHidden == true);
+	REQUIRE(Descriptor.bIsSealed == true);
+	REQUIRE(Descriptor.bExplicitlyLoaded == true);
+	REQUIRE(Descriptor.bHasExplicitPlatforms == true);
+
+	REQUIRE(Descriptor.LocalizationTargets.Num() == 1);
+	REQUIRE(Descriptor.LocalizationTargets[0].Name == TEXT("LocalizationTarget AAA"));
+	REQUIRE(Descriptor.LocalizationTargets[0].LoadingPolicy == ELocalizationTargetDescriptorLoadingPolicy::Always);
+	REQUIRE(Descriptor.LocalizationTargets[0].ConfigGenerationPolicy == ELocalizationConfigGenerationPolicy::Auto);
+
+	REQUIRE(Descriptor.Modules.Num() == 3);
+	REQUIRE(Descriptor.Modules[0].Name == TEXT("AAA"));
+	REQUIRE(Descriptor.Modules[0].Type == EHostType::DeveloperTool);
+	REQUIRE(Descriptor.Modules[0].LoadingPhase == ELoadingPhase::Default);
+	REQUIRE(Descriptor.Modules[0].bHasExplicitPlatforms == true);
+	REQUIRE(Descriptor.Modules[0].PlatformAllowList.Num() == 1);
+	REQUIRE(Descriptor.Modules[0].PlatformAllowList[0] == TEXT("Win64"));
+	REQUIRE(Descriptor.Modules[0].PlatformDenyList[0] == TEXT("Mac"));
+	REQUIRE(Descriptor.Modules[0].PlatformDenyList.Num() == 1);
+	REQUIRE(Descriptor.Modules[0].TargetAllowList.Num() == 1);
+	REQUIRE(Descriptor.Modules[0].TargetAllowList[0] == EBuildTargetType::Editor);
+	REQUIRE(Descriptor.Modules[0].TargetDenyList.Num() == 1);
+	REQUIRE(Descriptor.Modules[0].TargetDenyList[0] == EBuildTargetType::Client);
+	REQUIRE(Descriptor.Modules[0].TargetConfigurationAllowList.Num() == 1);
+	REQUIRE(Descriptor.Modules[0].TargetConfigurationAllowList[0] == EBuildConfiguration::Test);
+	REQUIRE(Descriptor.Modules[0].TargetConfigurationDenyList.Num() == 1);
+	REQUIRE(Descriptor.Modules[0].TargetConfigurationDenyList[0] == EBuildConfiguration::DebugGame);
+	REQUIRE(Descriptor.Modules[0].ProgramAllowList.Num() == 1);
+	REQUIRE(Descriptor.Modules[0].ProgramAllowList[0] == TEXT("UnrealInsights"));
+	REQUIRE(Descriptor.Modules[0].ProgramDenyList.Num() == 1);
+	REQUIRE(Descriptor.Modules[0].ProgramDenyList[0] == TEXT("LiveLinkHub"));
+	REQUIRE(Descriptor.Modules[0].GameTargetAllowList.Num() == 1);
+	REQUIRE(Descriptor.Modules[0].GameTargetAllowList[0] == TEXT("AAA"));
+	REQUIRE(Descriptor.Modules[0].GameTargetDenyList.Num() == 1);
+	REQUIRE(Descriptor.Modules[0].GameTargetDenyList[0] == TEXT("BBB"));
+	REQUIRE(Descriptor.Modules[0].AdditionalDependencies.Num() == 2);
+	REQUIRE(Descriptor.Modules[0].AdditionalDependencies[0] == TEXT("AAA"));
+	REQUIRE(Descriptor.Modules[0].AdditionalDependencies[1] == TEXT("BBB"));
+
+	REQUIRE(Descriptor.Plugins.Num() == 2);
+	REQUIRE(Descriptor.Plugins[0].Name == TEXT("DDD"));
+	REQUIRE(Descriptor.Plugins[0].bEnabled == true);
+	REQUIRE(Descriptor.Plugins[0].bOptional == true);
+	REQUIRE(Descriptor.Plugins[0].Description == TEXT("Ref Description goes here"));
+	REQUIRE(Descriptor.Plugins[0].MarketplaceURL == TEXT("Ref MarketplaceURL goes here"));
+	REQUIRE(Descriptor.Plugins[0].PlatformAllowList.Num() == 1);
+	REQUIRE(Descriptor.Plugins[0].PlatformAllowList[0] == TEXT("Win64"));
+	REQUIRE(Descriptor.Plugins[0].PlatformDenyList[0] == TEXT("Mac"));
+	REQUIRE(Descriptor.Plugins[0].PlatformDenyList.Num() == 1);
+	REQUIRE(Descriptor.Plugins[0].TargetAllowList.Num() == 1);
+	REQUIRE(Descriptor.Plugins[0].TargetAllowList[0] == EBuildTargetType::Editor);
+	REQUIRE(Descriptor.Plugins[0].TargetDenyList.Num() == 1);
+	REQUIRE(Descriptor.Plugins[0].TargetDenyList[0] == EBuildTargetType::Client);
+	REQUIRE(Descriptor.Plugins[0].TargetConfigurationAllowList.Num() == 1);
+	REQUIRE(Descriptor.Plugins[0].TargetConfigurationAllowList[0] == EBuildConfiguration::Test);
+	REQUIRE(Descriptor.Plugins[0].TargetConfigurationDenyList.Num() == 1);
+	REQUIRE(Descriptor.Plugins[0].TargetConfigurationDenyList[0] == EBuildConfiguration::DebugGame);
+
+	REQUIRE(Descriptor.PreBuildSteps.HostPlatformToCommands.Num() == 1);
+	REQUIRE(Descriptor.PreBuildSteps.HostPlatformToCommands.Contains(TEXT("PreBuildStepA")));
+	const TArray<FString>& PreBuildStepsA = Descriptor.PreBuildSteps.HostPlatformToCommands[TEXT("PreBuildStepA")];
+	REQUIRE(PreBuildStepsA.Num() == 2);
+	REQUIRE(PreBuildStepsA[0] == TEXT("Step 1a"));
+	REQUIRE(PreBuildStepsA[1] == TEXT("Step 2a"));
+	REQUIRE(Descriptor.PostBuildSteps.HostPlatformToCommands.Num() == 2);
+
+	REQUIRE(Descriptor.DisallowedPlugins.Num() == 1);
+	REQUIRE(Descriptor.DisallowedPlugins[0].Name == TEXT("ABC"));
+#if WITH_EDITOR
+	REQUIRE(Descriptor.DisallowedPlugins[0].Comment == TEXT("ABC Reasons"));
+#endif
+
+}
+
+TEST_CASE_NAMED(FProjectsLoadingProjectDescriptorTests, "System::Engine::Projects::Loading::ProjectDescriptor", "[ApplicationContextMask][SmokeFilter]")
+{
+	const TCHAR JsonText[] = TEXT(R"json(
+{
+	"FileVersion": 3,
+	"EngineAssociation": "EngineAssociation goes here",
+	"Category": "Category goes here",
+	"Description": "Description goes here",
+	"Enterprise": true,
+	"DisableEnginePluginsByDefault": true,
+	"AdditionalPluginDirectories": ["Content"],
+	"AdditionalRootDirectories": ["/RootDir"],
+	"TargetPlatforms": ["Windows", "Linux"],
+	"EpicSampleNameHash": 1234,
+	"Modules": [
+		{
+			"Name": "ModuleA",
+			"Type": "Editor",
+			"LoadingPhase": "Default"
+		}
+	],
+	"Plugins": [
+		{
+			"Name": "PluginA",
+			"Enabled": true
+		}
+	],
+	"PreBuildSteps": {
+		"PreBuildStepA": [
+			"Step 1a",
+			"Step 2a",
+		]
+	},
+	"PostBuildSteps": {
+		"PostBuildStepA": ["Step 1b"],
+		"PostBuildStepB": ["Step 2b"],
+	},	
+})json");
+
+
+	FProjectDescriptor Descriptor;
+
+	using namespace UE::Json;
+
+	TValueOrError<FDocument, FParseError> ParseResult = Parse(JsonText);
+	FString ParseErrorMessage = ParseResult.HasError() ? ParseResult.GetError().CreateMessage(JsonText) : FString();
+	REQUIRE_MESSAGE(ParseErrorMessage, ParseResult.HasValue());
+
+	TOptional<FConstObject> RootObject = GetRootObject(ParseResult.GetValue());
+	REQUIRE(RootObject);
+
+	FText FailReason;
+	REQUIRE_MESSAGE(FailReason.ToString(), FProjectsLoadingPluginDescriptorTestsHelper::Read(Descriptor, *RootObject, "UnitTest", FailReason));
+
+	REQUIRE(Descriptor.EngineAssociation == TEXT("EngineAssociation goes here"));
+	REQUIRE(Descriptor.Category == TEXT("Category goes here"));
+	REQUIRE(Descriptor.Description == TEXT("Description goes here"));
+	REQUIRE(Descriptor.bIsEnterpriseProject == true);
+	REQUIRE(Descriptor.bDisableEnginePluginsByDefault == true);
+	REQUIRE(Descriptor.GetAdditionalPluginDirectories().Num() == 1);
+#if WITH_EDITOR	
+	REQUIRE(Descriptor.GetAdditionalRootDirectories().Num() == 1);
+#endif
+	REQUIRE(Descriptor.TargetPlatforms.Num() >= 2);
+	REQUIRE(Descriptor.EpicSampleNameHash == 1234);
+
+	// their contents are already tested in the plugin loading test
+	REQUIRE(Descriptor.Modules.Num() == 1);
+	REQUIRE(Descriptor.Plugins.Num() == 1);
+	REQUIRE(Descriptor.PreBuildSteps.HostPlatformToCommands.Num() == 1);
+	REQUIRE(Descriptor.PostBuildSteps.HostPlatformToCommands.Num() == 2);
+
+
+}
+
+#endif // WITH_TESTS
diff --git a/Engine/Source/Runtime/Projects/Public/ProjectDescriptor.h b/Engine/Source/Runtime/Projects/Public/ProjectDescriptor.h
--- a/Engine/Source/Runtime/Projects/Public/ProjectDescriptor.h
+++ b/Engine/Source/Runtime/Projects/Public/ProjectDescriptor.h
@@ -12,6 +12,7 @@
 #include "PluginReferenceDescriptor.h"
 #include "Serialization/JsonWriter.h"
 #include "UObject/NameTypes.h"
+#include "JsonUtils/RapidJsonUtils.h"
 
 class FJsonObject;
 class FText;
@@ -201,4 +202,7 @@ private:
 	/** Set of project module names */
 	mutable TSet<FName> ModuleNamesCache;
 #endif //if WITH_EDITOR
+
+	friend struct FProjectsLoadingPluginDescriptorTestsHelper;
+	bool Read(const UE::Json::FConstObject Object, const FString& PathToProject, FText& OutFailReason);
 };
