[Backout] - CL44196605
#rnx
Original CL Desc
-----------------------------------------------------------------
Converted .uplugin and .uproject loading to RapidJson for performance reasons

diff --git a/Engine/Source/Runtime/Json/Private/JsonUtils/JsonConversion.cpp b/Engine/Source/Runtime/Json/Private/JsonUtils/JsonConversion.cpp
deleted file mode 100644
diff --git a/Engine/Source/Runtime/Json/Private/JsonUtils/RapidJsonUtils.cpp b/Engine/Source/Runtime/Json/Private/JsonUtils/RapidJsonUtils.cpp
deleted file mode 100644
diff --git a/Engine/Source/Runtime/Json/Private/Tests/JsonConversionTests.cpp b/Engine/Source/Runtime/Json/Private/Tests/JsonConversionTests.cpp
deleted file mode 100644
diff --git a/Engine/Source/Runtime/Json/Private/Tests/RapidJsonUtilsTests.cpp b/Engine/Source/Runtime/Json/Private/Tests/RapidJsonUtilsTests.cpp
deleted file mode 100644
diff --git a/Engine/Source/Runtime/Json/Public/JsonUtils/JsonConversion.h b/Engine/Source/Runtime/Json/Public/JsonUtils/JsonConversion.h
deleted file mode 100644
diff --git a/Engine/Source/Runtime/Json/Public/JsonUtils/RapidJsonUtils.h b/Engine/Source/Runtime/Json/Public/JsonUtils/RapidJsonUtils.h
deleted file mode 100644
diff --git a/Engine/Source/Runtime/Json/Public/RapidJsonIncludes.h b/Engine/Source/Runtime/Json/Public/RapidJsonIncludes.h
--- a/Engine/Source/Runtime/Json/Public/RapidJsonIncludes.h
+++ b/Engine/Source/Runtime/Json/Public/RapidJsonIncludes.h
@@ -3,7 +3,6 @@
 #pragma once
 
 #include "HAL/Platform.h"
-#include "HAL/UnrealMemory.h"
 #include "Misc/AssertionMacros.h"
 #include "Misc/CoreMiscDefines.h"
 
@@ -26,30 +25,6 @@ THIRD_PARTY_INCLUDES_START
 #define RAPIDJSON_ERROR_STRING(x) TEXT(x)
 #endif
 
-#ifndef RAPIDJSON_MALLOC
-#define RAPIDJSON_MALLOC FMemory::Malloc
-#endif
-
-#ifndef RAPIDJSON_REALLOC
-#define RAPIDJSON_REALLOC FMemory::Realloc
-#endif
-
-#ifndef RAPIDJSON_FREE
-#define RAPIDJSON_FREE FMemory::Free
-#endif
-
-#if PLATFORM_LITTLE_ENDIAN
-#define RAPIDJSON_ENDIAN RAPIDJSON_LITTLEENDIAN
-#else
-#define RAPIDJSON_ENDIAN RAPIDJSON_BIGENDIAN
-#endif
-
-#if PLATFORM_CPU_ARM_FAMILY
-#define RAPIDJSON_NEON
-#elif PLATFORM_CPU_X86_FAMILY
-#define RAPIDJSON_SSE42
-#endif
-
 #include "rapidjson/rapidjson.h"
 #include "rapidjson/document.h"
 #include "rapidjson/encodings.h"
diff --git a/Engine/Source/Runtime/Projects/Private/CustomBuildSteps.cpp b/Engine/Source/Runtime/Projects/Private/CustomBuildSteps.cpp
--- a/Engine/Source/Runtime/Projects/Private/CustomBuildSteps.cpp
+++ b/Engine/Source/Runtime/Projects/Private/CustomBuildSteps.cpp
@@ -3,8 +3,6 @@
 #include "CustomBuildSteps.h"
 #include "Dom/JsonObject.h"
 #include "Serialization/JsonSerializer.h"
-#include "JsonUtils/JsonConversion.h"
-#include "RapidJsonPluginLoading.h"
 
 #define LOCTEXT_NAMESPACE "ModuleDescriptor"
 
@@ -13,47 +11,30 @@ bool FCustomBuildSteps::IsEmpty() const
 	return HostPlatformToCommands.Num() == 0;
 }
 
-FCustomBuildSteps UE::Projects::Private::ReadCustomBuildSteps(UE::Json::FConstObject Object, const TCHAR* FieldName)
+void FCustomBuildSteps::Read(const FJsonObject& Object, const FString& FieldName)
 {
-	FCustomBuildSteps Result{};
-
-	TOptional<Json::FConstObject> StepsObject = UE::Json::GetObjectField(Object, FieldName);
-	if (StepsObject.IsSet())
+	TSharedPtr<FJsonValue> StepsValue = Object.TryGetField(FieldName);
+	if(StepsValue.IsValid() && StepsValue->Type == EJson::Object)
 	{
-		for (const Json::FMember& Member : *StepsObject)
+		const TSharedPtr<FJsonObject>& StepsObject = StepsValue->AsObject();
+		for(const TPair<FString, TSharedPtr<FJsonValue>>& HostPlatformAndSteps : StepsObject->Values)
 		{
-			TArray<FString>& Commands = Result.HostPlatformToCommands.FindOrAdd(Member.name.GetString());
-			if (Member.value.IsArray())
+			TArray<FString>& Commands = HostPlatformToCommands.FindOrAdd(HostPlatformAndSteps.Key);
+			if(HostPlatformAndSteps.Value.IsValid() && HostPlatformAndSteps.Value->Type == EJson::Array)
 			{
-				Json::FConstArray CommandsArray = Member.value.GetArray();
-				Commands.Reserve(CommandsArray.Size());
-				for (const Json::FValue& CommandValue : CommandsArray)
+				const TArray<TSharedPtr<FJsonValue>>& CommandsArray = HostPlatformAndSteps.Value->AsArray();
+				for(const TSharedPtr<FJsonValue>& CommandValue: CommandsArray)
 				{
-					if (CommandValue.IsString())
+					if(CommandValue->Type == EJson::String)
 					{
-						Commands.Emplace(CommandValue.GetString());
+						Commands.Add(CommandValue->AsString());
 					}
 				}
 			}
 		}
 	}
-
-	return Result;
-}
-
-void FCustomBuildSteps::Read(const FJsonObject& Object, const FString& FieldName)
-{
-	UE::Projects::Private::ReadFromDefaultJsonHelper(
-		Object,
-		[&FieldName, this](UE::Json::FConstObject RootObject)
-		{
-			*this = UE::Projects::Private::ReadCustomBuildSteps(RootObject, *FieldName);
-			return true;
-		}
-	);
 }
 
-
 void FCustomBuildSteps::Write(TJsonWriter<>& Writer, const FString& FieldName) const
 {
 	if (!IsEmpty())
diff --git a/Engine/Source/Runtime/Projects/Private/LocalizationDescriptor.cpp b/Engine/Source/Runtime/Projects/Private/LocalizationDescriptor.cpp
--- a/Engine/Source/Runtime/Projects/Private/LocalizationDescriptor.cpp
+++ b/Engine/Source/Runtime/Projects/Private/LocalizationDescriptor.cpp
@@ -6,8 +6,6 @@
 #include "Dom/JsonObject.h"
 #include "Serialization/JsonSerializer.h"
 #include "JsonUtils/JsonObjectArrayUpdater.h"
-#include "JsonUtils/JsonConversion.h"
-#include "RapidJsonPluginLoading.h"
 
 #define LOCTEXT_NAMESPACE "LocalizationDescriptor"
 
@@ -111,58 +109,98 @@ FLocalizationTargetDescriptor::FLocalizationTargetDescriptor(FString InName, ELo
 {
 }
 
-TOptional<FText> UE::Projects::Private::Read(UE::Json::FConstObject InObject, FLocalizationTargetDescriptor& Result)
+bool FLocalizationTargetDescriptor::Read(const FJsonObject& InObject, FText* OutFailReason /*= nullptr*/)
 {
 	// Read the target name
-	if (!TryGetStringField(InObject, TEXT("Name"), Result.Name))
+	TSharedPtr<FJsonValue> NameValue = InObject.TryGetField(TEXT("Name"));
+	if (!NameValue.IsValid() || NameValue->Type != EJson::String)
 	{
-		return LOCTEXT("TargetWithoutAName", "Found a 'Localization Target' entry with a missing 'Name' field");
+		if (OutFailReason)
+		{
+			*OutFailReason = LOCTEXT("TargetWithoutAName", "Found a 'Localization Target' entry with a missing 'Name' field");
 		}
+		return false;
+	}
+	Name = NameValue->AsString();
 
 	// Read the target loading policy
-	FString LoadingPolicyValue;
-	if (TryGetStringField(InObject, TEXT("LoadingPolicy"), LoadingPolicyValue))
+	TSharedPtr<FJsonValue> LoadingPolicyValue = InObject.TryGetField(TEXT("LoadingPolicy"));
+	if (LoadingPolicyValue.IsValid() && LoadingPolicyValue->Type == EJson::String)
 	{
-		Result.LoadingPolicy = ELocalizationTargetDescriptorLoadingPolicy::FromString(*LoadingPolicyValue);
-		if (Result.LoadingPolicy == ELocalizationTargetDescriptorLoadingPolicy::Max)
+		LoadingPolicy = ELocalizationTargetDescriptorLoadingPolicy::FromString(*LoadingPolicyValue->AsString());
+		if (LoadingPolicy == ELocalizationTargetDescriptorLoadingPolicy::Max)
 		{
-			return FText::Format(LOCTEXT("TargetWithInvalidLoadingPolicy", "Localization Target entry '{0}' specified an unrecognized target LoadingPolicy '{1}'"), FText::FromString(Result.Name), FText::FromString(LoadingPolicyValue));
+			if (OutFailReason)
+			{
+				*OutFailReason = FText::Format(LOCTEXT("TargetWithInvalidLoadingPolicy", "Localization Target entry '{0}' specified an unrecognized target LoadingPolicy '{1}'"), FText::FromString(Name), FText::FromString(LoadingPolicyValue->AsString()));
+			}
+			return false;
 		}
 	}
 
 		// Read the target generation policy
-	FString ConfigGenerationPolicyValue;
-	if (TryGetStringField(InObject, TEXT("ConfigGenerationPolicy"), ConfigGenerationPolicyValue))
+	TSharedPtr<FJsonValue> ConfigGenerationPolicyValue = InObject.TryGetField(TEXT("ConfigGenerationPolicy"));
+	if (ConfigGenerationPolicyValue.IsValid() && ConfigGenerationPolicyValue->Type == EJson::String)
 	{
-		Result.ConfigGenerationPolicy = ELocalizationConfigGenerationPolicy::FromString(*ConfigGenerationPolicyValue);
-		if (Result.ConfigGenerationPolicy == ELocalizationConfigGenerationPolicy::Max)
+		ConfigGenerationPolicy = ELocalizationConfigGenerationPolicy::FromString(*ConfigGenerationPolicyValue->AsString());
+		if (ConfigGenerationPolicy == ELocalizationConfigGenerationPolicy::Max)
 		{
-			return FText::Format(LOCTEXT("TargetWithInvalidGenerationPolicy", "Localization Target entry '{0}' specified an unrecognized target GenerationPolicy'{1}'"), FText::FromString(Result.Name), FText::FromString(ConfigGenerationPolicyValue));
+			if (OutFailReason)
+			{
+				*OutFailReason = FText::Format(LOCTEXT("TargetWithInvalidGenerationPolicy", "Localization Target entry '{0}' specified an unrecognized target GenerationPolicy'{1}'"), FText::FromString(Name), FText::FromString(ConfigGenerationPolicyValue->AsString()));
 			}
+			return false;
 		}
+	}
 	// If we can't find the value, GenerationPolicy already defaults to Never so it's ok 
 
-	return {};
-}
-
-bool FLocalizationTargetDescriptor::Read(const FJsonObject& InObject, FText* OutFailReason /*= nullptr*/)
-{
-	return UE::Projects::Private::ReadFromDefaultJson(InObject, *this, OutFailReason);
+	return true;
 }
 
 bool FLocalizationTargetDescriptor::Read(const FJsonObject& InObject, FText& OutFailReason)
 {
-	return UE::Projects::Private::ReadFromDefaultJson(InObject, *this, &OutFailReason);
+	return Read(InObject, &OutFailReason);
 }
 
 bool FLocalizationTargetDescriptor::ReadArray(const FJsonObject& InObject, const TCHAR* InName, TArray<FLocalizationTargetDescriptor>& OutTargets, FText* OutFailReason /*= nullptr*/)
 {
-	return UE::Projects::Private::ReadArrayFromDefaultJson(InObject, InName, OutTargets, OutFailReason);
+	bool bResult = true;
+
+	TSharedPtr<FJsonValue> TargetsArrayValue = InObject.TryGetField(InName);
+	if (TargetsArrayValue.IsValid() && TargetsArrayValue->Type == EJson::Array)
+	{
+		const TArray<TSharedPtr<FJsonValue>>& TargetsArray = TargetsArrayValue->AsArray();
+		for (const TSharedPtr<FJsonValue>& TargetValue : TargetsArray)
+		{
+			if (TargetValue.IsValid() && TargetValue->Type == EJson::Object)
+			{
+				FLocalizationTargetDescriptor Descriptor;
+				if (Descriptor.Read(*TargetValue->AsObject().Get(), OutFailReason))
+				{
+					OutTargets.Add(Descriptor);
+				}
+				else
+				{
+					bResult = false;
+				}
+			}
+			else
+			{
+				if (OutFailReason)
+				{
+					*OutFailReason = LOCTEXT("TargetWithInvalidTargetsArray", "The 'Localization Targets' array has invalid contents and was not able to be loaded.");
+				}
+				bResult = false;
+			}
+		}
+	}
+
+	return bResult;
 }
 
 bool FLocalizationTargetDescriptor::ReadArray(const FJsonObject& InObject, const TCHAR* InName, TArray<FLocalizationTargetDescriptor>& OutTargets, FText& OutFailReason)
 {
-	return UE::Projects::Private::ReadArrayFromDefaultJson(InObject, InName, OutTargets, &OutFailReason);
+	return ReadArray(InObject, InName, OutTargets, &OutFailReason);
 }
 
 void FLocalizationTargetDescriptor::Write(TJsonWriter<>& Writer) const
diff --git a/Engine/Source/Runtime/Projects/Private/ModuleDescriptor.cpp b/Engine/Source/Runtime/Projects/Private/ModuleDescriptor.cpp
--- a/Engine/Source/Runtime/Projects/Private/ModuleDescriptor.cpp
+++ b/Engine/Source/Runtime/Projects/Private/ModuleDescriptor.cpp
@@ -7,8 +7,7 @@
 #include "Serialization/JsonSerializer.h"
 #include "JsonUtils/JsonObjectArrayUpdater.h"
 #include "Modules/ModuleManager.h"
-#include "RapidJsonPluginLoading.h"
-#include "JsonUtils/JsonConversion.h"
+#include "JsonExtensions.h"
 
 #define LOCTEXT_NAMESPACE "ModuleDescriptor"
 
@@ -158,118 +157,158 @@ FModuleDescriptor::FModuleDescriptor(const FName InName, EHostType::Type InType,
 {
 }
 
-
-TOptional<FText> UE::Projects::Private::Read(UE::Json::FConstObject Object, FModuleDescriptor& Out)
+bool FModuleDescriptor::Read(const FJsonObject& Object, FText* OutFailReason /*= nullptr*/)
 {
 	// Read the module name
-	FString NameValue;
-	if (!TryGetStringField(Object, TEXT("Name"), NameValue))
+	TSharedPtr<FJsonValue> NameValue = Object.TryGetField(TEXT("Name"));
+	if(!NameValue.IsValid() || NameValue->Type != EJson::String)
 	{
-		return LOCTEXT("ModuleWithoutAName", "Found a 'Module' entry with a missing 'Name' field");
+		if (OutFailReason)
+		{
+			*OutFailReason = LOCTEXT("ModuleWithoutAName", "Found a 'Module' entry with a missing 'Name' field");
+		}
+		return false;
 	}
-	Out.Name = FName(*NameValue);
+	Name = FName(*NameValue->AsString());
 
 	// Read the module type
-	FString TypeValue;
-	if (!TryGetStringField(Object, TEXT("Type"), TypeValue))
+	TSharedPtr<FJsonValue> TypeValue = Object.TryGetField(TEXT("Type"));
+	if(!TypeValue.IsValid() || TypeValue->Type != EJson::String)
 	{
-		return FText::Format( LOCTEXT( "ModuleWithoutAType", "Found Module entry '{0}' with a missing 'Type' field" ), FText::FromName(Out.Name));
+		if (OutFailReason)
+		{
+			*OutFailReason = FText::Format( LOCTEXT( "ModuleWithoutAType", "Found Module entry '{0}' with a missing 'Type' field" ), FText::FromName(Name) );
+		}
+		return false;
 	}
-
-	Out.Type = EHostType::FromString(*TypeValue);
-	if (Out.Type == EHostType::Max)
+	Type = EHostType::FromString(*TypeValue->AsString());
+	if(Type == EHostType::Max)
 	{
-		return FText::Format( LOCTEXT( "ModuleWithInvalidType", "Module entry '{0}' specified an unrecognized module Type '{1}'" ), FText::FromName(Out.Name), FText::FromString(TypeValue) );
+		if (OutFailReason)
+		{
+			*OutFailReason = FText::Format( LOCTEXT( "ModuleWithInvalidType", "Module entry '{0}' specified an unrecognized module Type '{1}'" ), FText::FromName(Name), FText::FromString(TypeValue->AsString()) );
+		}
+		return false;
 	}
 
 	// Read the loading phase
-	FString LoadingPhaseValue;
-	if (TryGetStringField(Object, TEXT("LoadingPhase"), LoadingPhaseValue))
+	TSharedPtr<FJsonValue> LoadingPhaseValue = Object.TryGetField(TEXT("LoadingPhase"));
+	if(LoadingPhaseValue.IsValid() && LoadingPhaseValue->Type == EJson::String)
 	{
-		Out.LoadingPhase = ELoadingPhase::FromString(*LoadingPhaseValue);
-		if(Out.LoadingPhase == ELoadingPhase::Max)
+		LoadingPhase = ELoadingPhase::FromString(*LoadingPhaseValue->AsString());
+		if(LoadingPhase == ELoadingPhase::Max)
 		{
-			return FText::Format( LOCTEXT( "ModuleWithInvalidLoadingPhase", "Module entry '{0}' specified an unrecognized module LoadingPhase '{1}'" ), FText::FromName(Out.Name), FText::FromString(LoadingPhaseValue) );
+			if (OutFailReason)
+			{
+				*OutFailReason = FText::Format( LOCTEXT( "ModuleWithInvalidLoadingPhase", "Module entry '{0}' specified an unrecognized module LoadingPhase '{1}'" ), FText::FromName(Name), FText::FromString(LoadingPhaseValue->AsString()) );
+			}
+			return false;
 		}
 	}
 
 	// Read the allow/deny lists for platforms
-	TryGetStringArrayFieldWithDeprecatedFallback(Object, TEXT("PlatformAllowList"), TEXT("WhitelistPlatforms"), /*out*/ Out.PlatformAllowList);
-	TryGetStringArrayFieldWithDeprecatedFallback(Object, TEXT("PlatformDenyList"), TEXT("BlacklistPlatforms"), /*out*/ Out.PlatformDenyList);
-	TryGetBoolField(Object, TEXT("HasExplicitPlatforms"), Out.bHasExplicitPlatforms);
+	JsonExtensions::TryGetStringArrayFieldWithDeprecatedFallback(Object, TEXT("PlatformAllowList"), TEXT("WhitelistPlatforms"), /*out*/ PlatformAllowList);
+	JsonExtensions::TryGetStringArrayFieldWithDeprecatedFallback(Object, TEXT("PlatformDenyList"), TEXT("BlacklistPlatforms"), /*out*/ PlatformDenyList);
+	Object.TryGetBoolField(TEXT("HasExplicitPlatforms"), bHasExplicitPlatforms);
 
 	// Read the allow/deny lists for targets
-	TryGetEnumArrayFieldWithDeprecatedFallback(Object, TEXT("TargetAllowList"), TEXT("WhitelistTargets"), /*out*/ Out.TargetAllowList);
-	TryGetEnumArrayFieldWithDeprecatedFallback(Object, TEXT("TargetDenyList"), TEXT("BlacklistTargets"), /*out*/ Out.TargetDenyList);
+	JsonExtensions::TryGetEnumArrayFieldWithDeprecatedFallback(Object, TEXT("TargetAllowList"), TEXT("WhitelistTargets"), /*out*/ TargetAllowList);
+	JsonExtensions::TryGetEnumArrayFieldWithDeprecatedFallback(Object, TEXT("TargetDenyList"), TEXT("BlacklistTargets"), /*out*/ TargetDenyList);
 
 	// Read the allow/deny lists for target configurations
-	TryGetEnumArrayFieldWithDeprecatedFallback(Object, TEXT("TargetConfigurationAllowList"), TEXT("WhitelistTargetConfigurations"), /*out*/ Out.TargetConfigurationAllowList);
-	TryGetEnumArrayFieldWithDeprecatedFallback(Object, TEXT("TargetConfigurationDenyList"), TEXT("BlacklistTargetConfigurations"), /*out*/ Out.TargetConfigurationDenyList);
+	JsonExtensions::TryGetEnumArrayFieldWithDeprecatedFallback(Object, TEXT("TargetConfigurationAllowList"), TEXT("WhitelistTargetConfigurations"), /*out*/ TargetConfigurationAllowList);
+	JsonExtensions::TryGetEnumArrayFieldWithDeprecatedFallback(Object, TEXT("TargetConfigurationDenyList"), TEXT("BlacklistTargetConfigurations"), /*out*/ TargetConfigurationDenyList);
 
 	// Read the allow/deny lists for programs
-	TryGetStringArrayFieldWithDeprecatedFallback(Object, TEXT("ProgramAllowList"), TEXT("WhitelistPrograms"), /*out*/ Out.ProgramAllowList);
-	TryGetStringArrayFieldWithDeprecatedFallback(Object, TEXT("ProgramDenyList"), TEXT("BlacklistPrograms"), /*out*/ Out.ProgramDenyList);
+	JsonExtensions::TryGetStringArrayFieldWithDeprecatedFallback(Object, TEXT("ProgramAllowList"), TEXT("WhitelistPrograms"), /*out*/ ProgramAllowList);
+	JsonExtensions::TryGetStringArrayFieldWithDeprecatedFallback(Object, TEXT("ProgramDenyList"), TEXT("BlacklistPrograms"), /*out*/ ProgramDenyList);
 
 	// Read the additional dependencies
-	TryGetStringArrayField(Object, TEXT("AdditionalDependencies"), Out.AdditionalDependencies);
+	Object.TryGetStringArrayField(TEXT("AdditionalDependencies"), AdditionalDependencies);
 	
 	// Read the allow/deny list for game targets
-	TryGetStringArrayField(Object, TEXT("GameTargetAllowList"), Out.GameTargetAllowList);
-	TryGetStringArrayField(Object, TEXT("GameTargetDenyList"), Out.GameTargetDenyList);
+	Object.TryGetStringArrayField(TEXT("GameTargetAllowList"), GameTargetAllowList);
+	Object.TryGetStringArrayField(TEXT("GameTargetDenyList"), GameTargetDenyList);
 	
 	// Read the platform architecture allow/deny lists
-	auto ReadPlatformArchitectureList = [&Out, Object]( const TCHAR* FieldName, TMap<FString,TArray<FString>>& OutPlatformArchitectureList ) -> TOptional<FText>
+	auto ReadPlatformArchitectureList = [this, Object, OutFailReason]( const TCHAR* FieldName, TMap<FString,TArray<FString>>& OutPlatformArchitectureList )
 	{
 		TArray<FString> PlatformArchitecturePairs;
-		TryGetStringArrayField(Object, FieldName, PlatformArchitecturePairs);
+		Object.TryGetStringArrayField(FieldName, PlatformArchitecturePairs);
 		for (const FString& PlatformArchitecturePair : PlatformArchitecturePairs)
 		{
 			FString Platform, Architecture;
 			if (!PlatformArchitecturePair.Split(TEXT(":"), &Platform, &Architecture))
 			{
-				return FText::Format( LOCTEXT( "ModuleWithInvalidPlatformArchitecture", "Module entry '{0}' specified an unrecognized Platform:Architecture pair '{1}'" ), FText::FromName(Out.Name), FText::FromString(PlatformArchitecturePair) );
+				if (OutFailReason)
+				{
+					*OutFailReason = FText::Format( LOCTEXT( "ModuleWithInvalidPlatformArchitecture", "Module entry '{0}' specified an unrecognized Platform:Architecture pair '{1}'" ), FText::FromName(Name), FText::FromString(PlatformArchitecturePair) );
+				}
+				return false;
 			}
 
 			OutPlatformArchitectureList.FindOrAdd(Platform).AddUnique(Architecture);
 		}
+		return true;
 
-		return {};
 	};
-
-	TOptional<FText> ErrorResult = ReadPlatformArchitectureList(TEXT("PlatformArchitectureAllowList"), Out.PlatformArchitectureAllowList);
-	if (ErrorResult)
+	if (!ReadPlatformArchitectureList(TEXT("PlatformArchitectureAllowList"), PlatformArchitectureAllowList))
 	{
-		return ErrorResult;
+		return false;
 	}
-
-	ErrorResult = ReadPlatformArchitectureList(TEXT("PlatformArchitectureDenyList"), Out.PlatformArchitectureDenyList);
-
-	if (ErrorResult)
+	if (!ReadPlatformArchitectureList(TEXT("PlatformArchitectureDenyList"), PlatformArchitectureDenyList))
 	{
-		return ErrorResult;
+		return false;
 	}
 
-	return {};
-}
-	
-bool FModuleDescriptor::Read(const FJsonObject& Object, FText* OutFailReason /*= nullptr*/)
-{
-	return UE::Projects::Private::ReadFromDefaultJson(Object, *this, OutFailReason);
+	return true;
 }
 
 bool FModuleDescriptor::Read(const FJsonObject& Object, FText& OutFailReason)
 {
-	return UE::Projects::Private::ReadFromDefaultJson(Object, *this, &OutFailReason);
+	return Read(Object, &OutFailReason);
 }
 
 bool FModuleDescriptor::ReadArray(const FJsonObject& Object, const TCHAR* Name, TArray<FModuleDescriptor>& OutModules, FText* OutFailReason /*= nullptr*/)
 {
-	return UE::Projects::Private::ReadArrayFromDefaultJson(Object, Name, OutModules, OutFailReason);
+	bool bResult = true;
+
+	TSharedPtr<FJsonValue> ModulesArrayValue = Object.TryGetField(Name);
+	if(ModulesArrayValue.IsValid() && ModulesArrayValue->Type == EJson::Array)
+	{
+		const TArray< TSharedPtr< FJsonValue > >& ModulesArray = ModulesArrayValue->AsArray();
+		for(int Idx = 0; Idx < ModulesArray.Num(); Idx++)
+		{
+			const TSharedPtr<FJsonValue>& ModuleValue = ModulesArray[Idx];
+			if(ModuleValue.IsValid() && ModuleValue->Type == EJson::Object)
+			{
+				FModuleDescriptor Descriptor;
+				if(Descriptor.Read(*ModuleValue->AsObject().Get(), OutFailReason))
+				{
+					OutModules.Add(Descriptor);
+				}
+				else
+				{
+					bResult = false;
+				}
+			}
+			else
+			{
+				if (OutFailReason)
+				{
+					*OutFailReason = LOCTEXT( "ModuleWithInvalidModulesArray", "The 'Modules' array has invalid contents and was not able to be loaded." );
+				}
+				bResult = false;
+			}
+		}
+	}
+	
+	return bResult;
 }
 
 bool FModuleDescriptor::ReadArray(const FJsonObject& Object, const TCHAR* Name, TArray<FModuleDescriptor>& OutModules, FText& OutFailReason)
 {
-	return UE::Projects::Private::ReadArrayFromDefaultJson(Object, Name, OutModules, &OutFailReason);
+	return ReadArray(Object, Name, OutModules, &OutFailReason);
 }
 
 void FModuleDescriptor::Write(TJsonWriter<>& Writer) const
diff --git a/Engine/Source/Runtime/Projects/Private/PluginDescriptor.cpp b/Engine/Source/Runtime/Projects/Private/PluginDescriptor.cpp
--- a/Engine/Source/Runtime/Projects/Private/PluginDescriptor.cpp
+++ b/Engine/Source/Runtime/Projects/Private/PluginDescriptor.cpp
@@ -6,8 +6,6 @@
 #include "Serialization/JsonReader.h"
 #include "Serialization/JsonSerializer.h"
 #include "ProjectDescriptor.h"
-#include "JsonUtils/JsonConversion.h"
-#include "RapidJsonPluginLoading.h"
 
 #define LOCTEXT_NAMESPACE "PluginDescriptor"
 
@@ -124,48 +122,26 @@ bool FPluginDescriptor::Load(const FString& FileName, FText& OutFailReason)
 
 bool FPluginDescriptor::Read(const FString& Text, FText* OutFailReason /*= nullptr*/)
 {
-	using namespace UE::Json;
-
 #if WITH_EDITOR
 	CachedJson.Reset();
 	AdditionalFieldsToWrite.Reset();
 	AdditionalFieldsToRemove.Reset();
 #endif // WITH_EDITOR
 
-	TValueOrError<FDocument, FParseError> ParseResult = Parse(Text);
-
 	// Deserialize a JSON object from the string
-	if (ParseResult.HasError())
-	{
-		if (OutFailReason)
-	{
-			*OutFailReason = FText::Format(LOCTEXT("FailedToReadDescriptorFile", "Failed to read file. {0}"), FText::FromString(ParseResult.GetError().CreateMessage(Text)));
-		}
-		return false;
-	}
-
-	FDocument Document = ParseResult.StealValue();
-
-	TOptional<FConstObject> RootObject = GetRootObject(Document);
-	if (RootObject.IsSet())
+	TSharedPtr<FJsonObject> JsonObject = PluginDescriptor::DeserializeJson(Text, OutFailReason);
+	if (JsonObject.IsValid())
 	{
 		// Parse it as a plug-in descriptor
-		if (TOptional<FText> ReadError = UE::Projects::Private::Read(*RootObject, *this))
-		{
-			if (OutFailReason)
+		if (Read(*JsonObject.Get(), OutFailReason))
 		{
-				*OutFailReason = MoveTemp(*ReadError);
-			}			
-			return false;
-		}
-
 #if WITH_EDITOR
-		CachedJson = ConvertRapidJsonToSharedJsonObject(*RootObject);
+			CachedJson = JsonObject;
 			AdditionalFieldsToWrite.Reset();
 			AdditionalFieldsToRemove.Reset();
 #endif // WITH_EDITOR
 			return true;
-
+		}
 	}
 	return false;
 }
@@ -175,15 +151,19 @@ bool FPluginDescriptor::Read(const FString& Text, FText& OutFailReason)
 	return Read(Text, &OutFailReason);
 }
 
-TOptional<FText> UE::Projects::Private::Read(UE::Json::FConstObject Object, FPluginDescriptor& Result)
+bool FPluginDescriptor::Read(const FJsonObject& Object, FText* OutFailReason /*= nullptr*/)
 {
 	// Read the file version
 	int32 FileVersionInt32;
-	if(!TryGetNumberField(Object, TEXT("FileVersion"), FileVersionInt32))
+	if(!Object.TryGetNumberField(TEXT("FileVersion"), FileVersionInt32))
 	{
-		if(!TryGetNumberField(Object, TEXT("PluginFileVersion"), FileVersionInt32))
+		if(!Object.TryGetNumberField(TEXT("PluginFileVersion"), FileVersionInt32))
 		{
-			return LOCTEXT("InvalidProjectFileVersion", "File does not have a valid 'FileVersion' number.");
+			if (OutFailReason)
+			{
+				*OutFailReason = LOCTEXT("InvalidProjectFileVersion", "File does not have a valid 'FileVersion' number.");
+			}
+			return false;
 		}
 	}
 
@@ -191,149 +171,155 @@ TOptional<FText> UE::Projects::Private::Read(UE::Json::FConstObject Object, FPlu
 	EPluginDescriptorVersion PluginFileVersion = (EPluginDescriptorVersion)FileVersionInt32;
 	if ((PluginFileVersion <= EPluginDescriptorVersion::Invalid) || (PluginFileVersion > EPluginDescriptorVersion::Latest))
 	{
+		if (OutFailReason)
+		{
 			const FText ReadVersionText = FText::FromString(FString::Printf(TEXT("%d"), (int32)PluginFileVersion));
 			const FText LatestVersionText = FText::FromString(FString::Printf(TEXT("%d"), (int32)EPluginDescriptorVersion::Latest));
-		return FText::Format(LOCTEXT("ProjectFileVersionTooLarge", "File appears to be in a newer version ({0}) of the file format that we can load (max version: {1})."), ReadVersionText, LatestVersionText);
+			*OutFailReason = FText::Format(LOCTEXT("ProjectFileVersionTooLarge", "File appears to be in a newer version ({0}) of the file format that we can load (max version: {1})."), ReadVersionText, LatestVersionText);
+		}
+		return false;
 	}
 
 	// Read the other fields
-	TryGetNumberField(Object, TEXT("Version"), Result.Version);
-	TryGetStringField(Object, TEXT("VersionName"), Result.VersionName);
-	TryGetStringField(Object, TEXT("FriendlyName"), Result.FriendlyName);
-	TryGetStringField(Object, TEXT("Description"), Result.Description);
+	Object.TryGetNumberField(TEXT("Version"), Version);
+	Object.TryGetStringField(TEXT("VersionName"), VersionName);
+	Object.TryGetStringField(TEXT("FriendlyName"), FriendlyName);
+	Object.TryGetStringField(TEXT("Description"), Description);
 
-	if (!TryGetStringField(Object, TEXT("Category"), Result.Category))
+	if (!Object.TryGetStringField(TEXT("Category"), Category))
 	{
 		// Category used to be called CategoryPath in .uplugin files
-		TryGetStringField(Object, TEXT("CategoryPath"), Result.Category);
+		Object.TryGetStringField(TEXT("CategoryPath"), Category);
 	}
         
 	// Due to a difference in command line parsing between Windows and Mac, we shipped a few Mac samples containing
 	// a category name with escaped quotes. Remove them here to make sure we can list them in the right category.
-	if (Result.Category.Len() >= 2 && Result.Category.StartsWith(TEXT("\""), ESearchCase::CaseSensitive) && Result.Category.EndsWith(TEXT("\""), ESearchCase::CaseSensitive))
+	if (Category.Len() >= 2 && Category.StartsWith(TEXT("\""), ESearchCase::CaseSensitive) && Category.EndsWith(TEXT("\""), ESearchCase::CaseSensitive))
 	{
-		Result.Category.MidInline(1, Result.Category.Len() - 2, EAllowShrinking::No);
+		Category.MidInline(1, Category.Len() - 2, EAllowShrinking::No);
 	}
 
-	TryGetStringField(Object, TEXT("CreatedBy"), Result.CreatedBy);
-	TryGetStringField(Object, TEXT("CreatedByURL"), Result.CreatedByURL);
-	TryGetStringField(Object, TEXT("DocsURL"), Result.DocsURL);
-	TryGetStringField(Object, TEXT("MarketplaceURL"), Result.MarketplaceURL);
-	TryGetStringField(Object, TEXT("SupportURL"), Result.SupportURL);
-	TryGetStringField(Object, TEXT("EngineVersion"), Result.EngineVersion);
-	TryGetStringField(Object, TEXT("EditorCustomVirtualPath"), Result.EditorCustomVirtualPath);
-	TryGetStringArrayField(Object, TEXT("SupportedTargetPlatforms"), Result.SupportedTargetPlatforms);
-	TryGetStringArrayField(Object, TEXT("SupportedPrograms"), Result.SupportedPrograms);
-	TryGetBoolField(Object, TEXT("bIsPluginExtension"), Result.bIsPluginExtension);
+	Object.TryGetStringField(TEXT("CreatedBy"), CreatedBy);
+	Object.TryGetStringField(TEXT("CreatedByURL"), CreatedByURL);
+	Object.TryGetStringField(TEXT("DocsURL"), DocsURL);
+	Object.TryGetStringField(TEXT("MarketplaceURL"), MarketplaceURL);
+	Object.TryGetStringField(TEXT("SupportURL"), SupportURL);
+	Object.TryGetStringField(TEXT("EngineVersion"), EngineVersion);
+	Object.TryGetStringField(TEXT("EditorCustomVirtualPath"), EditorCustomVirtualPath);
+	Object.TryGetStringArrayField(TEXT("SupportedTargetPlatforms"), SupportedTargetPlatforms);
+	Object.TryGetStringArrayField(TEXT("SupportedPrograms"), SupportedPrograms);
+	Object.TryGetBoolField(TEXT("bIsPluginExtension"), bIsPluginExtension);
 
-	if (TOptional<FText> ReadError = ReadArray(Object, TEXT("Modules"), Result.Modules))
+	if (!FModuleDescriptor::ReadArray(Object, TEXT("Modules"), Modules, OutFailReason))
 	{
-		return ReadError;
+		return false;
 	}
 
-	if (TOptional<FText> ReadError = ReadArray(Object, TEXT("LocalizationTargets"), Result.LocalizationTargets))
+	if (!FLocalizationTargetDescriptor::ReadArray(Object, TEXT("LocalizationTargets"), LocalizationTargets, OutFailReason))
 	{
-		return ReadError;
+		return false;
 	}
 
-	TryGetStringField(Object, TEXT("VersePath"), Result.VersePath);
-	TryGetStringField(Object, TEXT("DeprecatedEngineVersion"), Result.DeprecatedEngineVersion);
+	Object.TryGetStringField(TEXT("VersePath"), VersePath);
+	Object.TryGetStringField(TEXT("DeprecatedEngineVersion"), DeprecatedEngineVersion);
 
 	// Read the Verse scope
-	Json::FValue::ConstMemberIterator VerseScopeValue = Object.FindMember(TEXT("VerseScope"));
-	if (VerseScopeValue != Object.MemberEnd() && VerseScopeValue->value.IsString())
+	TSharedPtr<FJsonValue> VerseScopeValue = Object.TryGetField(TEXT("VerseScope"));
+	if (VerseScopeValue.IsValid() && VerseScopeValue->Type == EJson::String)
 	{
-		if (TOptional<EVerseScope::Type> MaybeVerseScope = EVerseScope::FromString(VerseScopeValue->value.GetString()))
+		if (TOptional<EVerseScope::Type> MaybeVerseScope = EVerseScope::FromString(*VerseScopeValue->AsString()))
 		{
-			Result.VerseScope = *MaybeVerseScope;
+			VerseScope = *MaybeVerseScope;
 		}
 		else
 		{
-			return FText::Format(LOCTEXT("PluginWithInvalidVerseScope", "Plugin entry 'VerseScope' specified an unrecognized value '{0}'"), FText::FromString(VerseScopeValue->value.GetString()));
+			if (OutFailReason)
+			{
+				*OutFailReason = FText::Format(LOCTEXT("PluginWithInvalidVerseScope", "Plugin entry 'VerseScope' specified an unrecognized value '{1}'"), FText::FromString(VerseScopeValue->AsString()));
+			}
+			return false;
 		}
 	}
 
 	// Read the Verse version.
-	Json::FValue::ConstMemberIterator VerseVersionValue = Object.FindMember(TEXT("VerseVersion"));
-	if (VerseVersionValue != Object.MemberEnd())
+	TSharedPtr<FJsonValue> VerseVersionValue = Object.TryGetField(TEXT("VerseVersion"));
+	if (VerseVersionValue.IsValid())
 	{
-		if (VerseVersionValue->value.IsInt())
+		uint32 PluginVerseVersion;
+		if (VerseVersionValue->TryGetNumber(PluginVerseVersion))
 		{
-			Result.VerseVersion = VerseVersionValue->value.GetInt();
+			VerseVersion = PluginVerseVersion;
 		}
 		else
 		{
-			return FText::Format(LOCTEXT("PluginWithInvalidVerseVersion", "Plugin entry 'VerseVersion' specified an unrecognized type '{0}'"), FText::FromString(Json::GetValueTypeName(VerseVersionValue->value)));
+			if (OutFailReason)
+			{
+				*OutFailReason = FText::Format(LOCTEXT("PluginWithInvalidVerseVersion", "Plugin entry 'VerseVersion' specified an unrecognized value '{1}'"), FText::FromString(VerseVersionValue->AsString()));
+			}
+			return false;
 		}
 	}
 
-	TryGetBoolField(Object, TEXT("EnableSceneGraph"), Result.bEnableSceneGraph);
+	Object.TryGetBoolField(TEXT("EnableSceneGraph"), bEnableSceneGraph);
 	
-	TryGetBoolField(Object, TEXT("EnableVerseAssetReflection"), Result.bEnableVerseAssetReflection);
+	Object.TryGetBoolField(TEXT("EnableVerseAssetReflection"), bEnableVerseAssetReflection);
 
-	TryGetBoolField(Object, TEXT("EnableIAD"), Result.bEnableIAD);
+	Object.TryGetBoolField(TEXT("EnableIAD"), bEnableIAD);
 
 	bool bEnabledByDefault;
-	if(TryGetBoolField(Object, TEXT("EnabledByDefault"), bEnabledByDefault))
+	if(Object.TryGetBoolField(TEXT("EnabledByDefault"), bEnabledByDefault))
 	{
-		Result.EnabledByDefault = bEnabledByDefault ? EPluginEnabledByDefault::Enabled : EPluginEnabledByDefault::Disabled;
+		EnabledByDefault = bEnabledByDefault ? EPluginEnabledByDefault::Enabled : EPluginEnabledByDefault::Disabled;
 	}
 
-	TryGetBoolField(Object, TEXT("CanContainContent"), Result.bCanContainContent);
-	TryGetBoolField(Object, TEXT("CanContainVerse"), Result.bCanContainVerse);
-	TryGetBoolField(Object, TEXT("NoCode"), Result.bNoCode);
-	TryGetBoolField(Object, TEXT("IsBetaVersion"), Result.bIsBetaVersion);
-	TryGetBoolField(Object, TEXT("IsExperimentalVersion"), Result.bIsExperimentalVersion);
-	TryGetBoolField(Object, TEXT("Installed"), Result.bInstalled);
-	TryGetBoolField(Object, TEXT("RequiresBuildPlatform"), Result.bRequiresBuildPlatform);
-	TryGetBoolField(Object, TEXT("Hidden"), Result.bIsHidden);
-	TryGetBoolField(Object, TEXT("Sealed"), Result.bIsSealed);
-	TryGetBoolField(Object, TEXT("ExplicitlyLoaded"), Result.bExplicitlyLoaded);
-	TryGetBoolField(Object, TEXT("HasExplicitPlatforms"), Result.bHasExplicitPlatforms);
+	Object.TryGetBoolField(TEXT("CanContainContent"), bCanContainContent);
+	Object.TryGetBoolField(TEXT("CanContainVerse"), bCanContainVerse);
+	Object.TryGetBoolField(TEXT("NoCode"), bNoCode);
+	Object.TryGetBoolField(TEXT("IsBetaVersion"), bIsBetaVersion);
+	Object.TryGetBoolField(TEXT("IsExperimentalVersion"), bIsExperimentalVersion);
+	Object.TryGetBoolField(TEXT("Installed"), bInstalled);
+	Object.TryGetBoolField(TEXT("RequiresBuildPlatform"), bRequiresBuildPlatform);
+	Object.TryGetBoolField(TEXT("Hidden"), bIsHidden);
+	Object.TryGetBoolField(TEXT("Sealed"), bIsSealed);
+	Object.TryGetBoolField(TEXT("ExplicitlyLoaded"), bExplicitlyLoaded);
+	Object.TryGetBoolField(TEXT("HasExplicitPlatforms"), bHasExplicitPlatforms);
 
 	bool bCanBeUsedWithUnrealHeaderTool;
-	if(TryGetBoolField(Object, TEXT("CanBeUsedWithUnrealHeaderTool"), bCanBeUsedWithUnrealHeaderTool) && bCanBeUsedWithUnrealHeaderTool)
+	if(Object.TryGetBoolField(TEXT("CanBeUsedWithUnrealHeaderTool"), bCanBeUsedWithUnrealHeaderTool) && bCanBeUsedWithUnrealHeaderTool)
 	{
-		Result.SupportedPrograms.Add(TEXT("UnrealHeaderTool"));
+		SupportedPrograms.Add(TEXT("UnrealHeaderTool"));
 	}
 
-	Result.PreBuildSteps = ReadCustomBuildSteps(Object, TEXT("PreBuildSteps"));
-	Result.PostBuildSteps = ReadCustomBuildSteps(Object, TEXT("PostBuildSteps"));
+	PreBuildSteps.Read(Object, TEXT("PreBuildSteps"));
+	PostBuildSteps.Read(Object, TEXT("PostBuildSteps"));
 
-	if (TOptional<FText> ReadError = ReadArray(Object, TEXT("Plugins"), Result.Plugins))
+	if (!FPluginReferenceDescriptor::ReadArray(Object, TEXT("Plugins"), Plugins, OutFailReason))
 	{
-		return ReadError;
+		return false;
 	}
 
 	// Backwards compatibility support
-	if (TArray<FString> DisallowedPluginNameStrings; TryGetStringArrayField(Object, TEXT("DisallowedPlugins"), DisallowedPluginNameStrings))
+	TArray<FString> DisallowedPluginNameStrings;
+	if (Object.TryGetStringArrayField(TEXT("DisallowedPlugins"), DisallowedPluginNameStrings))
 	{
-		Result.DisallowedPlugins.Reserve(DisallowedPluginNameStrings.Num());
+		DisallowedPlugins.Reserve(DisallowedPluginNameStrings.Num());
 		for (int32 Index = 0; Index < DisallowedPluginNameStrings.Num(); ++Index)
 		{
-			FPluginDisallowedDescriptor& PluginDisallowedDescriptor = Result.DisallowedPlugins.AddDefaulted_GetRef();
+			FPluginDisallowedDescriptor& PluginDisallowedDescriptor = DisallowedPlugins.AddDefaulted_GetRef();
 			PluginDisallowedDescriptor.Name = DisallowedPluginNameStrings[Index];
 		}
 	}
-	else
+	else if (!FPluginDisallowedDescriptor::ReadArray(Object, TEXT("DisallowedPlugins"), DisallowedPlugins, OutFailReason))
 	{
-		if (TOptional<FText> ReadError = ReadArray(Object, TEXT("DisallowedPlugins"), Result.DisallowedPlugins))
-		{
-			return ReadError;
-		}
+		return false;
 	}
 
-	return {};
-}
-
-bool FPluginDescriptor::Read(const FJsonObject& Object, FText* OutFailReason /*= nullptr*/)
-{
-	return UE::Projects::Private::ReadFromDefaultJson(Object, *this, OutFailReason);
+	return true;
 }
 
 bool FPluginDescriptor::Read(const FJsonObject& Object, FText& OutFailReason)
 {
-	return UE::Projects::Private::ReadFromDefaultJson(Object, *this, &OutFailReason);
+	return Read(Object, &OutFailReason);
 }
 
 bool FPluginDescriptor::Save(const TCHAR* FileName, FText* OutFailReason /*= nullptr*/) const
diff --git a/Engine/Source/Runtime/Projects/Private/PluginDisallowedDescriptor.cpp b/Engine/Source/Runtime/Projects/Private/PluginDisallowedDescriptor.cpp
--- a/Engine/Source/Runtime/Projects/Private/PluginDisallowedDescriptor.cpp
+++ b/Engine/Source/Runtime/Projects/Private/PluginDisallowedDescriptor.cpp
@@ -7,8 +7,6 @@
 #include "JsonUtils/JsonObjectArrayUpdater.h"
 #include "ProjectDescriptor.h"
 #include "JsonExtensions.h"
-#include "JsonUtils/JsonConversion.h"
-#include "RapidJsonPluginLoading.h"
 
 #define LOCTEXT_NAMESPACE "PluginDisallowedDescriptor"
 
@@ -30,24 +28,25 @@ namespace PluginDisallowedDescriptor
 	}
 }
 
-TOptional<FText> UE::Projects::Private::Read(UE::Json::FConstObject Object, FPluginDisallowedDescriptor& Result)
+bool FPluginDisallowedDescriptor::Read(const FJsonObject& Object, FText* OutFailReason /*= nullptr*/)
 {
+	bool bSuccess = true;
+
 	// Get the name
-	if (!TryGetStringField(Object, TEXT("Name"), Result.Name))
+	if (!Object.TryGetStringField(TEXT("Name"), Name))
 	{
-		return LOCTEXT("FPluginDisallowedWithoutName", "Disallowed plugin must have a 'Name' field");
+		if (OutFailReason)
+		{
+			*OutFailReason = LOCTEXT("FPluginDisallowedWithoutName", "Disallowed plugin must have a 'Name' field");
 		}
+		bSuccess = false;
+	}
 
 #if WITH_EDITOR
-	TryGetStringField(Object, TEXT("Comment"), Result.Comment);
+	Object.TryGetStringField(TEXT("Comment"), Comment);
 #endif //if WITH_EDITOR
 
-	return {};
-}
-
-bool FPluginDisallowedDescriptor::Read(const FJsonObject& Object, FText* OutFailReason /*= nullptr*/)
-{
-	return UE::Projects::Private::ReadFromDefaultJson(Object, *this, OutFailReason);
+	return bSuccess;
 }
 
 void FPluginDisallowedDescriptor::Write(TJsonWriter<>& Writer) const
@@ -70,7 +69,29 @@ void FPluginDisallowedDescriptor::UpdateJson(FJsonObject& JsonObject) const
 
 bool FPluginDisallowedDescriptor::ReadArray(const FJsonObject& Object, const TCHAR* Name, TArray<FPluginDisallowedDescriptor>& OutPlugins, FText* OutFailReason /*= nullptr*/)
 {
-	return UE::Projects::Private::ReadArrayFromDefaultJson(Object, Name, OutPlugins, OutFailReason);
+	const TArray< TSharedPtr<FJsonValue> >* Array;
+
+	if (Object.TryGetArrayField(Name, Array))
+	{
+		for (const TSharedPtr<FJsonValue>& Item : *Array)
+		{
+			const TSharedPtr<FJsonObject>* ItemObject = nullptr;
+
+			if (Item.IsValid() && Item->TryGetObject(ItemObject) && ItemObject && ItemObject->IsValid())
+			{
+				FPluginDisallowedDescriptor DisallowedRef;
+
+				if (!DisallowedRef.Read(*ItemObject->Get(), OutFailReason))
+				{
+					return false;
+				}
+
+				OutPlugins.Add(MoveTemp(DisallowedRef));
+			}
+		}
+	}
+
+	return true;
 }
 
 void FPluginDisallowedDescriptor::WriteArray(TJsonWriter<>& Writer, const TCHAR* ArrayName, const TArray<FPluginDisallowedDescriptor>& Plugins)
diff --git a/Engine/Source/Runtime/Projects/Private/PluginManifest.cpp b/Engine/Source/Runtime/Projects/Private/PluginManifest.cpp
--- a/Engine/Source/Runtime/Projects/Private/PluginManifest.cpp
+++ b/Engine/Source/Runtime/Projects/Private/PluginManifest.cpp
@@ -4,78 +4,9 @@
 #include "Misc/FileHelper.h"
 #include "Serialization/JsonReader.h"
 #include "Serialization/JsonSerializer.h"
-#include "JsonUtils/JsonConversion.h"
-#include "Misc/Paths.h"
-#include "RapidJsonPluginLoading.h"
 
 #define LOCTEXT_NAMESPACE "PluginManifest"
 
-namespace PluginManifest
-{
-	void DumpContents(const TArray<FPluginManifestEntry>& Contents, const FString& Filename)
-	{
-		TArray<FString> Lines;
-
-		for (const FPluginManifestEntry& Entry : Contents)
-		{
-			FString Line;
-			Entry.Descriptor.Write(Line);
-
-			Line = Entry.File + TEXT("\n") + Line;
-
-			Lines.Add(MoveTemp(Line));
-		}
-
-		Lines.Sort();
-
-		FString DumpFilename = FPaths::ProjectSavedDir() / Filename;
-		FFileHelper::SaveStringArrayToFile(Lines, *DumpFilename, FFileHelper::EEncodingOptions::ForceUTF8WithoutBOM);
-	}
-}
-
-TOptional<FText> UE::Projects::Private::Read(UE::Json::FConstObject Object, FPluginManifest& Result)
-{
-	using namespace UE::Json;
-
-	TOptional<FConstArray> JsonContents = GetArrayField(Object, TEXT("Contents"));
-	if (!JsonContents.IsSet())
-	{
-		return LOCTEXT("FailedToReadDescriptorFile", "No 'Contents' array found");
-	}
-
-	Result.Contents.Reserve(JsonContents->Size());
-
-	for (const FValue& JsonEntry : *JsonContents)
-	{
-		if (!JsonEntry.IsObject())
-		{
-			continue;
-		}
-
-		FPluginManifestEntry& ManifestEntry = Result.Contents.Emplace_GetRef();
-
-		FConstObject JsonEntryObject = JsonEntry.GetObject();
-		TOptional<FConstObject> JsonDescriptor = GetObjectField(JsonEntryObject, TEXT("Descriptor"));
-
-		if (!JsonDescriptor.IsSet())
-		{
-			return FText::Format(LOCTEXT("FailedToReadDescriptorFile", "Plugin descriptor index '{0}' missing 'Descriptor' field"), Result.Contents.Num()-1);
-		}
-
-		if (!TryGetStringField(JsonEntryObject, TEXT("File"), ManifestEntry.File))
-		{
-			return FText::Format(LOCTEXT("FailedToReadDescriptorFile", "Plugin descriptor index '{0}' missing 'File' field"), Result.Contents.Num()-1);
-		}
-
-		if (TOptional<FText> ReadError = UE::Projects::Private::Read(*JsonDescriptor, ManifestEntry.Descriptor))
-		{
-			return FText::Format(LOCTEXT("FailedToReadDescriptorFile", "Failed to parse plugin descriptor entry for file '{0}': {1}"), FText::FromString(ManifestEntry.File), *ReadError);
-		}
-	}
-
-	return {};
-}
-
 bool FPluginManifest::Load(const FString& FileName, FText& OutFailReason)
 {
 	// Read the file to a string
@@ -86,44 +17,35 @@ bool FPluginManifest::Load(const FString& FileName, FText& OutFailReason)
 		return false;
 	}
 
+	// Deserialize a JSON object from the string
+	TSharedPtr< FJsonObject > ObjectPtr;
+	TSharedRef< TJsonReader<> > Reader = TJsonReaderFactory<>::Create(FileContents);
+	if (!FJsonSerializer::Deserialize(Reader, ObjectPtr) || !ObjectPtr.IsValid() )
 	{
-		SCOPED_BOOT_TIMING("FPluginManifest::Load RapidJson")
-
-		// grab the characters so we can parse insitu safely (FString API says adding non-terminating zeros can be unsafe so just take the data)
-		TArray<TCHAR> Characters = MoveTemp(FileContents.GetCharArray());
-		TValueOrError<UE::Json::FDocument, UE::Json::FParseError> Document = UE::Json::ParseInPlace(Characters);
-
-		if (Document.HasError())
-		{
-			// have to re-parse the string because it was destructively altered
-			ensure(FFileHelper::LoadFileToString(FileContents, *FileName));
-			OutFailReason = FText::Format(LOCTEXT("FailedToReadDescriptorFile", "Failed to read file. {0}"), FText::FromString(Document.GetError().CreateMessage(FileContents)));
-			
+		OutFailReason = FText::Format(LOCTEXT("FailedToReadDescriptorFile", "Failed to read file. {0}"), FText::FromString(Reader->GetErrorMessage()));
 		return false;
 	}
 
-		TOptional<UE::Json::FConstObject> RootObject = UE::Json::GetRootObject(Document.GetValue());
-		if (RootObject.IsSet())
-		{
-			//for (int32 N = 0; N < 1000; ++N)
+	// Parse the object
+	return Read(*ObjectPtr.Get(), OutFailReason);
+}
+
+bool FPluginManifest::Read(const FJsonObject& Object, FText& OutFailReason)
+{
+	TArray<TSharedPtr<FJsonValue>> JsonContents = Object.GetArrayField(TEXT("Contents"));
+	for (const TSharedPtr<FJsonValue>& JsonEntryValue : JsonContents)
 	{
-				if (TOptional<FText> Result = UE::Projects::Private::Read(*RootObject, *this))
+		TSharedPtr<FJsonObject> JsonEntry = JsonEntryValue->AsObject();
+
+		FPluginManifestEntry Entry;
+		Entry.File = JsonEntry->GetStringField(TEXT("File"));
+		if (!Entry.Descriptor.Read(*JsonEntry->GetObjectField(TEXT("Descriptor")).Get(), OutFailReason))
 		{
-					OutFailReason = *Result;
 			return false;
 		}
+		Contents.Add(Entry);
 	}
-
 	return true;
-		}
-	}
-
-	return false;
-}
-
-bool FPluginManifest::Read(const FJsonObject& Object, FText& OutFailReason)
-{
-	return UE::Projects::Private::ReadFromDefaultJson(Object, *this, &OutFailReason);
 }
 
 #undef LOCTEXT_NAMESPACE
diff --git a/Engine/Source/Runtime/Projects/Private/PluginReferenceDescriptor.cpp b/Engine/Source/Runtime/Projects/Private/PluginReferenceDescriptor.cpp
--- a/Engine/Source/Runtime/Projects/Private/PluginReferenceDescriptor.cpp
+++ b/Engine/Source/Runtime/Projects/Private/PluginReferenceDescriptor.cpp
@@ -6,8 +6,7 @@
 #include "Serialization/JsonSerializer.h"
 #include "JsonUtils/JsonObjectArrayUpdater.h"
 #include "ProjectDescriptor.h"
-#include "RapidJsonPluginLoading.h"
-#include "JsonUtils/JsonConversion.h"
+#include "JsonExtensions.h"
 
 #define LOCTEXT_NAMESPACE "PluginDescriptor"
 
@@ -119,88 +118,121 @@ bool FPluginReferenceDescriptor::IsSupportedTargetPlatform(const FString& Platfo
 	}
 }
 
-TOptional<FText> UE::Projects::Private::Read(UE::Json::FConstObject Object, FPluginReferenceDescriptor& Result)
+bool FPluginReferenceDescriptor::Read(const TSharedRef<FJsonObject>& Object, FText* OutFailReason /*= nullptr*/)
 {
-	using namespace UE::Json;
+PRAGMA_DISABLE_DEPRECATION_WARNINGS
+	return Read(*Object, OutFailReason, Object);
+PRAGMA_ENABLE_DEPRECATION_WARNINGS
+}
 
+bool FPluginReferenceDescriptor::Read(const FJsonObject& Object, FText* OutFailReason /*= nullptr*/, TSharedPtr<FJsonObject> ObjectPtr /*= nullptr*/)
+{
 #if WITH_EDITOR
-	Result.CachedJson = ConvertRapidJsonToSharedJsonObject(Object);
-	Result.AdditionalFieldsToWrite.Reset();
+	CachedJson = ObjectPtr;
+	AdditionalFieldsToWrite.Reset();
 #endif // WITH_EDITOR
 
+	bool bSuccess = true;
+
 	// Get the name
-	if (!TryGetStringField(Object, TEXT("Name"), Result.Name))
+	if (!Object.TryGetStringField(TEXT("Name"), Name))
 	{
-		return LOCTEXT("PluginReferenceWithoutName", "Plugin references must have a 'Name' field");
+		if (OutFailReason)
+		{
+			*OutFailReason = LOCTEXT("PluginReferenceWithoutName", "Plugin references must have a 'Name' field");
+		}
+		bSuccess = false;
 	}
 
 	// Get the enabled field
-	if (!TryGetBoolField(Object, TEXT("Enabled"), Result.bEnabled))
+	if (!Object.TryGetBoolField(TEXT("Enabled"), bEnabled))
 	{
-		return LOCTEXT("PluginReferenceWithoutEnabled", "Plugin references must have an 'Enabled' field");
+		if (OutFailReason)
+		{
+			*OutFailReason = LOCTEXT("PluginReferenceWithoutEnabled", "Plugin references must have an 'Enabled' field");
+		}
+		bSuccess = false;
 	}
 
 	// Read the optional field
-	TryGetBoolField(Object, TEXT("Optional"), Result.bOptional);
+	Object.TryGetBoolField(TEXT("Optional"), bOptional);
 
-	TryGetBoolField(Object, TEXT("Activate"), Result.bActivate);
+	Object.TryGetBoolField(TEXT("Activate"), bActivate);
 
 	// Read the metadata for users that don't have the plugin installed
-	TryGetStringField(Object, TEXT("Description"), Result.Description);
-	TryGetStringField(Object, TEXT("MarketplaceURL"), Result.MarketplaceURL);
+	Object.TryGetStringField(TEXT("Description"), Description);
+	Object.TryGetStringField(TEXT("MarketplaceURL"), MarketplaceURL);
 
 	// Get the platform lists
-	TryGetStringArrayFieldWithDeprecatedFallback(Object, TEXT("PlatformAllowList"), TEXT("WhitelistPlatforms"), /*out*/ Result.PlatformAllowList);
-	TryGetStringArrayFieldWithDeprecatedFallback(Object, TEXT("PlatformDenyList"), TEXT("BlacklistPlatforms"), /*out*/ Result.PlatformDenyList);
+	JsonExtensions::TryGetStringArrayFieldWithDeprecatedFallback(Object, TEXT("PlatformAllowList"), TEXT("WhitelistPlatforms"), /*out*/ PlatformAllowList);
+	JsonExtensions::TryGetStringArrayFieldWithDeprecatedFallback(Object, TEXT("PlatformDenyList"), TEXT("BlacklistPlatforms"), /*out*/ PlatformDenyList);
 
 	// Get the target configuration lists
-	TryGetEnumArrayFieldWithDeprecatedFallback(Object, TEXT("TargetConfigurationAllowList"), TEXT("WhitelistTargetConfigurations"), /*out*/ Result.TargetConfigurationAllowList);
-	TryGetEnumArrayFieldWithDeprecatedFallback(Object, TEXT("TargetConfigurationDenyList"), TEXT("BlacklistTargetConfigurations"), /*out*/ Result.TargetConfigurationDenyList);
+	JsonExtensions::TryGetEnumArrayFieldWithDeprecatedFallback(Object, TEXT("TargetConfigurationAllowList"), TEXT("WhitelistTargetConfigurations"), /*out*/ TargetConfigurationAllowList);
+	JsonExtensions::TryGetEnumArrayFieldWithDeprecatedFallback(Object, TEXT("TargetConfigurationDenyList"), TEXT("BlacklistTargetConfigurations"), /*out*/ TargetConfigurationDenyList);
 
 	// Get the target lists
-	TryGetEnumArrayFieldWithDeprecatedFallback(Object, TEXT("TargetAllowList"), TEXT("WhitelistTargets"), /*out*/ Result.TargetAllowList);
-	TryGetEnumArrayFieldWithDeprecatedFallback(Object, TEXT("TargetDenyList"), TEXT("BlacklistTargets"), /*out*/ Result.TargetDenyList);
+	JsonExtensions::TryGetEnumArrayFieldWithDeprecatedFallback(Object, TEXT("TargetAllowList"), TEXT("WhitelistTargets"), /*out*/ TargetAllowList);
+	JsonExtensions::TryGetEnumArrayFieldWithDeprecatedFallback(Object, TEXT("TargetDenyList"), TEXT("BlacklistTargets"), /*out*/ TargetDenyList);
 
 	// Get the supported platform list
-	TryGetStringArrayField(Object, TEXT("SupportedTargetPlatforms"), Result.SupportedTargetPlatforms);
-	TryGetBoolField(Object, TEXT("HasExplicitPlatforms"), Result.bHasExplicitPlatforms);
+	Object.TryGetStringArrayField(TEXT("SupportedTargetPlatforms"), SupportedTargetPlatforms);
+	Object.TryGetBoolField(TEXT("HasExplicitPlatforms"), bHasExplicitPlatforms);
 
 	int32 ReadVersion;
-	if (TryGetNumberField(Object, TEXT("Version"), ReadVersion))
+	if (Object.TryGetNumberField(TEXT("Version"), ReadVersion))
 	{
-		Result.RequestedVersion = ReadVersion;
-		if (!Result.bEnabled)
+		RequestedVersion = ReadVersion;
+		if (!bEnabled)
 		{
-			return LOCTEXT("PluginReferenceDisabledWithVersion", "Plugin references cannot be used to disable explicit versions. Remove the 'Version' field when 'Enabled' is false.");
+			if (bSuccess && OutFailReason)
+			{
+				*OutFailReason = LOCTEXT("PluginReferenceDisabledWithVersion", "Plugin references cannot be used to disable explicit versions. Remove the 'Version' field when 'Enabled' is false.");
 			}
+			bSuccess = false;
 		}
+	}
 
-	return {};
-}
-
-bool FPluginReferenceDescriptor::Read(const TSharedRef<FJsonObject>& Object, FText* OutFailReason /*= nullptr*/)
-{
-	return UE::Projects::Private::ReadFromDefaultJson(*Object, *this, OutFailReason);
-}
-
-bool FPluginReferenceDescriptor::Read(const FJsonObject& Object, FText* OutFailReason /*= nullptr*/, TSharedPtr<FJsonObject> ObjectPtr /*= nullptr*/)
-{
-	return UE::Projects::Private::ReadFromDefaultJson(Object, *this, OutFailReason);
+	return bSuccess;
 }
 
 bool FPluginReferenceDescriptor::Read(const FJsonObject& Object, FText& OutFailReason, TSharedPtr<FJsonObject> ObjectPtr /*= nullptr*/)
 {
-	return UE::Projects::Private::ReadFromDefaultJson(Object, *this, &OutFailReason);
+PRAGMA_DISABLE_DEPRECATION_WARNINGS
+	return Read(Object, &OutFailReason, ObjectPtr);
+PRAGMA_ENABLE_DEPRECATION_WARNINGS
 }
 
 bool FPluginReferenceDescriptor::ReadArray(const FJsonObject& Object, const TCHAR* Name, TArray<FPluginReferenceDescriptor>& OutPlugins, FText* OutFailReason /*= nullptr*/)
 {
-	return UE::Projects::Private::ReadArrayFromDefaultJson(Object, Name, OutPlugins, OutFailReason);
+	const TArray< TSharedPtr<FJsonValue> >* Array;
+
+	if (Object.TryGetArrayField(Name, Array))
+	{
+		for (const TSharedPtr<FJsonValue> &Item : *Array)
+		{
+			const TSharedPtr<FJsonObject>* ItemObject = nullptr;
+
+			if (Item.IsValid() && Item->TryGetObject(ItemObject) && ItemObject && ItemObject->IsValid())
+			{
+				FPluginReferenceDescriptor PluginRef;
+
+				if (!PluginRef.Read(ItemObject->ToSharedRef(), OutFailReason))
+				{
+					return false;
+				}
+
+				OutPlugins.Add(PluginRef);
+			}
+		}
+	}
+
+	return true;
 }
 
 bool FPluginReferenceDescriptor::ReadArray(const FJsonObject& Object, const TCHAR* Name, TArray<FPluginReferenceDescriptor>& OutPlugins, FText& OutFailReason)
 {
-	return UE::Projects::Private::ReadArrayFromDefaultJson(Object, Name, OutPlugins, &OutFailReason);
+	return ReadArray(Object, Name, OutPlugins, &OutFailReason);
 }
 
 void FPluginReferenceDescriptor::Write(TJsonWriter<>& Writer) const
diff --git a/Engine/Source/Runtime/Projects/Private/ProjectDescriptor.cpp b/Engine/Source/Runtime/Projects/Private/ProjectDescriptor.cpp
--- a/Engine/Source/Runtime/Projects/Private/ProjectDescriptor.cpp
+++ b/Engine/Source/Runtime/Projects/Private/ProjectDescriptor.cpp
@@ -8,8 +8,6 @@
 #include "Serialization/JsonSerializer.h"
 #include "Interfaces/IProjectManager.h"
 #include "Misc/DataDrivenPlatformInfoRegistry.h"
-#include "JsonUtils/JsonConversion.h"
-#include "RapidJsonPluginLoading.h"
 
 #define LOCTEXT_NAMESPACE "ProjectDescriptor"
 
@@ -57,8 +55,6 @@ void FProjectDescriptor::UpdateSupportedTargetPlatforms(const FName& InPlatformN
 
 bool FProjectDescriptor::Load(const FString& FileName, FText& OutFailReason)
 {
-	SCOPED_BOOT_TIMING("FProjectDescriptor::Load")
-
 	// Read the file to a string
 	FString FileContents;
 	if (!FFileHelper::LoadFileToString(FileContents, *FileName))
@@ -69,42 +65,28 @@ bool FProjectDescriptor::Load(const FString& FileName, FText& OutFailReason)
 		return false;
 	}
 
-	// grab the characters so we can parse insitu safely (FString API says adding non-terminating zeros can be unsafe so just take the data)
-	TArray<TCHAR> Characters = MoveTemp(FileContents.GetCharArray());
-	TValueOrError<UE::Json::FDocument, UE::Json::FParseError> ParseResult = UE::Json::ParseInPlace(Characters);
-
-	if (ParseResult.HasError())
+	// Deserialize a JSON object from the string
+	TSharedPtr< FJsonObject > Object;
+	TSharedRef< TJsonReader<> > Reader = TJsonReaderFactory<>::Create(FileContents);
+	if ( !FJsonSerializer::Deserialize(Reader, Object) || !Object.IsValid() )
 	{
-		// have to re-load the file because it was destructively altered
-		ensure(FFileHelper::LoadFileToString(FileContents, *FileName));
-
 		// Do not try to localize these messages, a missing project file usually indicates missing data and an
 		// access to ICU will result in failure that's harder to diagnose
-		OutFailReason = FText::FromString(FString::Printf(TEXT("Failed to read file %s: %s"), *FileName, *ParseResult.GetError().CreateMessage(FileContents)));
-		return false;
-	}
-
-	TOptional<UE::Json::FConstObject> RootObject = UE::Json::GetRootObject(ParseResult.GetValue());
-	if (!RootObject.IsSet())
-	{
+		OutFailReason = FText::FromString(FString::Printf(TEXT("Failed to read file %s"), *FileName));
 		return false;
 	}
 
-	return Read(*RootObject, FPaths::GetPath(FileName), OutFailReason);
+	// Parse it as a project descriptor
+	return Read(*Object.Get(), FPaths::GetPath(FileName), OutFailReason);
 }
 
-bool FProjectDescriptor::Read(const UE::Json::FConstObject Object, const FString& PathToProject, FText& OutFailReason)
+bool FProjectDescriptor::Read(const FJsonObject& Object, const FString& PathToProject, FText& OutFailReason)
 {
-	SCOPED_BOOT_TIMING("FProjectDescriptor::Read")
-
-	using namespace UE::Json;
-	using namespace UE::Projects::Private;
-
 	// Read the file version
 	int32 FileVersionInt32;
-	if(!TryGetNumberField(Object, TEXT("FileVersion"), FileVersionInt32))
+	if(!Object.TryGetNumberField(TEXT("FileVersion"), FileVersionInt32))
 	{
-		if(!TryGetNumberField(Object, TEXT("ProjectFileVersion"), FileVersionInt32))
+		if(!Object.TryGetNumberField(TEXT("ProjectFileVersion"), FileVersionInt32))
 		{
 			OutFailReason = LOCTEXT("InvalidProjectFileVersion", "File does not have a valid 'FileVersion' number.");
 			return false;
@@ -122,45 +104,48 @@ bool FProjectDescriptor::Read(const UE::Json::FConstObject Object, const FString
 	}
 
 	// Read simple fields
-	TryGetStringField(Object, TEXT("EngineAssociation"), EngineAssociation);
-	TryGetStringField(Object, TEXT("Category"), Category);
-	TryGetStringField(Object, TEXT("Description"), Description);
-	TryGetBoolField(Object, TEXT("Enterprise"), bIsEnterpriseProject);
-	TryGetBoolField(Object, TEXT("DisableEnginePluginsByDefault"), bDisableEnginePluginsByDefault);
+	Object.TryGetStringField(TEXT("EngineAssociation"), EngineAssociation);
+	Object.TryGetStringField(TEXT("Category"), Category);
+	Object.TryGetStringField(TEXT("Description"), Description);
+	Object.TryGetBoolField(TEXT("Enterprise"), bIsEnterpriseProject);
+	Object.TryGetBoolField(TEXT("DisableEnginePluginsByDefault"), bDisableEnginePluginsByDefault);
 
 #if WITH_EDITOR
 	ModuleNamesCache.Reset();
 #endif
 
-	if (TOptional<FText> ReadModulesError = ReadArray(Object, TEXT("Modules"), Modules))
+	// Read the modules
+	if(!FModuleDescriptor::ReadArray(Object, TEXT("Modules"), Modules, OutFailReason))
 	{
-		OutFailReason = *ReadModulesError;
 		return false;
 	}
 
-	if (TOptional<FText> ReadPluginsError = ReadArray(Object, TEXT("Plugins"), Plugins))
+	// Read the plugins
+	if(!FPluginReferenceDescriptor::ReadArray(Object, TEXT("Plugins"), Plugins, OutFailReason))
 	{
-		OutFailReason = *ReadPluginsError;
 		return false;
 	}
 
 	// Read the list of additional plugin directories to scan
-	TArray<FString> AdditionalPluginDirectoriesValue;
-	if (TryGetStringArrayField(Object, TEXT("AdditionalPluginDirectories"), AdditionalPluginDirectoriesValue))
+	const TArray< TSharedPtr<FJsonValue> >* AdditionalPluginDirectoriesValue;
+	if (Object.TryGetArrayField(TEXT("AdditionalPluginDirectories"), AdditionalPluginDirectoriesValue))
 	{
 #if WITH_EDITOR || (IS_PROGRAM && WITH_PLUGIN_SUPPORT)
-		for (FString& AdditionalPluginDir : AdditionalPluginDirectoriesValue)
+		for (int32 Idx = 0; Idx < AdditionalPluginDirectoriesValue->Num(); Idx++)
 		{
+			FString AdditionalPluginDir;
+			if ((*AdditionalPluginDirectoriesValue)[Idx]->TryGetString(AdditionalPluginDir))
+			{
 				if (FPaths::IsRelative(AdditionalPluginDir))
 				{
 					AdditionalPluginDir = IFileManager::Get().ConvertToAbsolutePathForExternalAppForWrite(*(PathToProject / AdditionalPluginDir));
 				}
 				AddPluginDirectory(AdditionalPluginDir);
 			}
-
+		}
 #endif //if WITH_EDITOR
 		// If this is a packaged build and there are additional directories, they need to be remapped to the packaged location
-		if (FPlatformProperties::RequiresCookedData() && AdditionalPluginDirectoriesValue.Num() > 0)
+		if (FPlatformProperties::RequiresCookedData() && AdditionalPluginDirectoriesValue->Num() > 0)
 		{
 			AdditionalPluginDirectories.Empty();
 			FString RemappedDir = FPaths::ProjectDir() / TEXT("../RemappedPlugins/");
@@ -174,11 +159,14 @@ bool FProjectDescriptor::Read(const UE::Json::FConstObject Object, const FString
 
 #if WITH_EDITOR
 	// Read the list of additional root directories to scan
-	TArray<FString> AdditionalRootDirectoriesValue;
-	if (TryGetStringArrayField(Object, TEXT("AdditionalRootDirectories"), AdditionalRootDirectoriesValue))
+	const TArray< TSharedPtr<FJsonValue> >* AdditionalRootDirectoriesValue;
+	if (Object.TryGetArrayField(TEXT("AdditionalRootDirectories"), AdditionalRootDirectoriesValue))
 	{
-		for (FString& AdditionalRoot : AdditionalRootDirectoriesValue)
+		for (const TSharedPtr<FJsonValue>& AdditionalRootDirectoryValue : *AdditionalRootDirectoriesValue)
 		{
+			FString AdditionalRoot;
+			if (AdditionalRootDirectoryValue->TryGetString(AdditionalRoot))
+			{
 				if (FPaths::IsRelative(AdditionalRoot))
 				{
 					AdditionalRoot = IFileManager::Get().ConvertToAbsolutePathForExternalAppForWrite(*(PathToProject / AdditionalRoot));
@@ -186,10 +174,22 @@ bool FProjectDescriptor::Read(const UE::Json::FConstObject Object, const FString
 				AddRootDirectory(AdditionalRoot);
 			}
 		}
+	}
 #endif //if WITH_EDITOR
 
 	// Read the target platforms
-	TryGetStringArrayField(Object, TEXT("TargetPlatforms"), TargetPlatforms);
+	const TArray< TSharedPtr<FJsonValue> > *TargetPlatformsValue;
+	if(Object.TryGetArrayField(TEXT("TargetPlatforms"), TargetPlatformsValue))
+	{
+		for(int32 Idx = 0; Idx < TargetPlatformsValue->Num(); Idx++)
+		{
+			FString TargetPlatform;
+			if((*TargetPlatformsValue)[Idx]->TryGetString(TargetPlatform))
+			{
+				TargetPlatforms.Add(*TargetPlatform);
+			}
+		}
+	}
 
 	// check if the project has directories for extended platforms, and assume support if it does
 	// however, if there were no platforms already listed, then all platforms are supported, and we don't
@@ -220,26 +220,15 @@ bool FProjectDescriptor::Read(const UE::Json::FConstObject Object, const FString
 	}
 
 	// Get the sample name hash
-	TryGetNumberField(Object, TEXT("EpicSampleNameHash"), EpicSampleNameHash);
+	Object.TryGetNumberField(TEXT("EpicSampleNameHash"), EpicSampleNameHash);
 
 	// Read the custom build steps
-	PreBuildSteps = ReadCustomBuildSteps(Object, TEXT("PreBuildSteps"));
-	PostBuildSteps = ReadCustomBuildSteps(Object, TEXT("PostBuildSteps"));
+	PreBuildSteps.Read(Object, TEXT("PreBuildSteps"));
+	PostBuildSteps.Read(Object, TEXT("PostBuildSteps"));
 
 	return true;
 }
 
-bool FProjectDescriptor::Read(const FJsonObject& Object, const FString& PathToProject, FText& OutFailReason)
-{
-	return UE::Projects::Private::ReadFromDefaultJsonHelper(
-		Object,
-		[&PathToProject, &OutFailReason, this](UE::Json::FConstObject Object)
-		{
-			return this->Read(Object, PathToProject, OutFailReason);
-		}
-	);
-}
-
 bool FProjectDescriptor::Save(const FString& FileName, FText& OutFailReason)
 {
 	if (IProjectManager::Get().IsSuppressingProjectFileWrite())
diff --git a/Engine/Source/Runtime/Projects/Private/RapidJsonPluginLoading.cpp b/Engine/Source/Runtime/Projects/Private/RapidJsonPluginLoading.cpp
deleted file mode 100644
diff --git a/Engine/Source/Runtime/Projects/Private/RapidJsonPluginLoading.h b/Engine/Source/Runtime/Projects/Private/RapidJsonPluginLoading.h
deleted file mode 100644
diff --git a/Engine/Source/Runtime/Projects/Private/Tests/RapidJsonPluginLoadingTests.cpp b/Engine/Source/Runtime/Projects/Private/Tests/RapidJsonPluginLoadingTests.cpp
deleted file mode 100644
diff --git a/Engine/Source/Runtime/Projects/Public/ProjectDescriptor.h b/Engine/Source/Runtime/Projects/Public/ProjectDescriptor.h
--- a/Engine/Source/Runtime/Projects/Public/ProjectDescriptor.h
+++ b/Engine/Source/Runtime/Projects/Public/ProjectDescriptor.h
@@ -12,7 +12,6 @@
 #include "PluginReferenceDescriptor.h"
 #include "Serialization/JsonWriter.h"
 #include "UObject/NameTypes.h"
-#include "JsonUtils/RapidJsonUtils.h"
 
 class FJsonObject;
 class FText;
@@ -202,7 +201,4 @@ private:
 	/** Set of project module names */
 	mutable TSet<FName> ModuleNamesCache;
 #endif //if WITH_EDITOR
-
-	friend struct FProjectsLoadingPluginDescriptorTestsHelper;
-	bool Read(const UE::Json::FConstObject Object, const FString& PathToProject, FText& OutFailReason);
 };
