Split up ConfigureEnabledPlugins, as VC was throwing a too-large-to-optimize error in some cases. Added a context struct to pass containers around

diff --git a/Engine/Source/Runtime/Projects/Private/PluginManager.cpp b/Engine/Source/Runtime/Projects/Private/PluginManager.cpp
--- a/Engine/Source/Runtime/Projects/Private/PluginManager.cpp
+++ b/Engine/Source/Runtime/Projects/Private/PluginManager.cpp
@@ -12,6 +12,7 @@
 #include "Misc/Paths.h"
 #include "Misc/ConfigCacheIni.h"
 #include "Misc/ConfigContext.h"
+#include "Misc/ConfigUtilities.h"
 #include "Misc/ScopedSlowTask.h"
 #include "Misc/CoreDelegates.h"
 #include "Misc/App.h"
@@ -1362,237 +1363,115 @@ bool FPluginManager::GetPreloadBinaries()
 	return bPreloadedBinaries;
 }
 
-void FPluginManager::ProcessEnabledPlugins()
+#if READ_TARGET_ENABLED_PLUGINS_FROM_RECEIPT
+static TUniquePtr<FTargetReceipt> FindMatchingTargetFile (const FString& ReceiptWildcard)
 {
-	SCOPED_BOOT_TIMING("ProcessEnabledPlugins");
-
-	// generate optimal list of plugins to process
-	TArray<TSharedRef<FPlugin>> PluginsArray;
-	for (const FString& PluginName : PluginsToConfigure)
-	{
-		TSharedRef<FPlugin> Plugin = DiscoveredPluginMapUtils::FindPluginInMap_Checked(AllPlugins, PluginName);
-		
-		// only process enabled plugins
-		if (Plugin->bEnabled && !Plugin->Descriptor.bExplicitlyLoaded)
-		{
-			PluginsArray.Add(Plugin);
-		}
-	}
-
-	TArray<FString> ConfigFilesPluginsCannotOverride;
-	GConfig->GetArray(TEXT("Plugins"), TEXT("ConfigFilesPluginsCannotOverride"), ConfigFilesPluginsCannotOverride, GEngineIni);
+	FString DefaultEditorTarget;
+	GConfig->GetString(TEXT("/Script/BuildSettings.BuildSettings"), TEXT("DefaultEditorTarget"), DefaultEditorTarget, GEngineIni);
 	
-	FCriticalSection ConfigCS;
-	FCriticalSection PluginPakCS;
-	
-	// Mount all the enabled plugins
-	FConfigModificationTracker ChangeTracker;
-	ChangeTracker.bTrackLoadedFiles = true;
+	TArray<FString> AllTargetFilesWithoutPath;
+	const FString ReceiptPath = FPaths::GetPath(ReceiptWildcard);
+	IFileManager::Get().FindFiles(AllTargetFilesWithoutPath, *ReceiptWildcard, true, false);
 
-	TArray<FName> PluginNames;
-	PluginNames.Reserve(PluginsArray.Num());
-	
-	// walk over each plugin, and add the config files that are named for the plugin
-	// this is a separate loop so that a plugin can modify another plugin's configs in the second loop below
-	for (TSharedRef<FPlugin> PluginPtr : PluginsArray)
-	{
-		FPlugin& Plugin = *PluginPtr;
-		FName PluginName(*Plugin.Name);
-		PluginNames.Add(PluginName);
-		
-		UE_LOG(LogPluginManager, Log, TEXT("Mounting %s plugin %s"), *EnumToString(Plugin.Type), *Plugin.GetName());
-		UE_LOG(LogPluginManager, Verbose, TEXT("Plugin path: %s"), *Plugin.FileName);
-
-		UE::PluginManager::Private::CheckForDeprecatedPlugin(Plugin);
+	// Used as a fallback in lieu of an exact match (only for EBuildTargetType::Editor).
+	TPair<FString, TUniquePtr<FTargetReceipt>> BestMatch;
 
-		// register this plugin, so it can be found right now, but also when loading other platforms, or when modifying it by other plugins
-		const bool bIncludePluginNameInBranchName = false;
-		FConfigCacheIni::RegisterPlugin(PluginName, Plugin.GetBaseDir(), Plugin.GetExtensionBaseDirs(), DynamicLayerPriority::Plugin, bIncludePluginNameInBranchName);
-		
-		FConfigContext Context = FConfigContext::ReadIntoGConfig();
-		Context.ChangeTracker = &ChangeTracker;
-		Context.ConfigFileTag = *Plugin.Name;
-		Context.Load(*Plugin.Name);
-	}
-	
-	// gang up dynamic layer adds as much as possible to reduce the amount of "CalcDiff" calls etc 
-	FConfigCacheIni::AddMultiplePluginsToAllBranches(PluginNames);
-
-	// walk over each plugin, modify existing ini's, mount the content, etc
-	ParallelFor(PluginsArray.Num(), [&PluginsArray, &ConfigCS, &PluginPakCS, &ConfigFilesPluginsCannotOverride, &ChangeTracker, this](int32 Index)
+	for (const FString& TargetFileWithoutPath : AllTargetFilesWithoutPath)
 	{
-		FPlugin& Plugin = *PluginsArray[Index];
-		FName PluginName(*Plugin.Name);
-
-		// Build the list of content folders
-		if (Plugin.Descriptor.bCanContainContent)
+		FString TargetFile = FPaths::Combine(ReceiptPath, TargetFileWithoutPath);
+		TUniquePtr<FTargetReceipt> Receipt = MakeUnique<FTargetReceipt>();
+		if (Receipt->Read(TargetFile))
 		{
+			if (Receipt->TargetType == FApp::GetBuildTargetType() && Receipt->Configuration == FApp::GetBuildConfiguration())
 			{
-				FScopeLock Locker(&ConfigCS);
-
-				// we need to look up the section each time because other loops could add entries
-				if (FConfigFile* EngineConfigFile = GConfig->FindConfigFile(GEngineIni))
+				// Priority:
+				// - If we are non-editor, first non-editor receipt
+				// - Otherwise, first that matches the running executable
+				// - Otherwise, first corresponding to DefaultEditorTarget (if specified)
+				// - Otherwise, first editor receipt
+				if (Receipt->TargetType != EBuildTargetType::Editor || Receipt->LaunchesCurrentExecutable())
 				{
-					EngineConfigFile->AddUniqueToSection(TEXT("Core.System"), "Paths", Plugin.GetContentDir());
+					UE_LOG(LogPluginManager, Log, TEXT("Found matching target receipt: %s"), *TargetFile);
+					return Receipt;
 				}
-			}
 
-			TArray<FString>	FoundPaks;
-			FPakFileSearchVisitor PakVisitor(FoundPaks);
-			IPlatformFile& PlatformFile = FPlatformFileManager::Get().GetPlatformFile();
-
-			// Pak files are loaded from <PluginName>/Content/Paks/<PlatformName>
-			if (FPlatformProperties::RequiresCookedData())
-			{
-				PlatformFile.IterateDirectoryRecursively(*(Plugin.GetContentDir() / TEXT("Paks") / FPlatformProperties::PlatformName()), PakVisitor);
-
-				for (const FString& PakPath : FoundPaks)
+				const bool bIsDefaultTarget = (DefaultEditorTarget.Len() != 0) && (DefaultEditorTarget == Receipt->TargetName);
+				if (bIsDefaultTarget || !BestMatch.Value)
 				{
-					FScopeLock Locker(&PluginPakCS);
-					if (FCoreDelegates::MountPak.IsBound())
-					{
-						FCoreDelegates::MountPak.Execute(PakPath, 0);
-						PluginsWithPakFile.AddUnique(PluginsArray[Index]);
-					}
-					else
-					{
-						UE_LOG(LogPluginManager, Warning, TEXT("PAK file (%s) could not be mounted because MountPak is not bound"), *PakPath);
-					}
+					BestMatch = { MoveTemp(TargetFile), MoveTemp(Receipt) };
+					// `TargetFile` and `Receipt` are now invalidated by virtue of being moved-from.
 				}
 			}
 		}
-	}, true); // @todo disable parallelism for now as it's causing hard to track problems
-}
+	}
 
-bool FPluginManager::ConfigureEnabledPlugins()
-{
-#if (WITH_ENGINE && !IS_PROGRAM) || WITH_PLUGIN_SUPPORT
-	if(PluginsToConfigure.Num() > 0)
+	if (BestMatch.Value)
 	{
-		SCOPED_BOOT_TIMING("FPluginManager::ConfigureEnabledPlugins");
-
-		bHaveAllRequiredPlugins = false;
-
-		// Set of all the plugins which have been enabled
-		TMap<FString, FPlugin*> EnabledPlugins;
-
-		// Keep a set of all the plugin names that have been configured. We read configuration data from different places, but only configure a plugin from the first place that it's referenced.
-		TSet<FString> ConfiguredPluginNames;
-
-		// Keep the list of newly available localization targets
-		TArray<FString> AdditionalLocResPaths;
-
-#if READ_TARGET_ENABLED_PLUGINS_FROM_RECEIPT
-		FString DefaultEditorTarget;
-		GConfig->GetString(TEXT("/Script/BuildSettings.BuildSettings"), TEXT("DefaultEditorTarget"), DefaultEditorTarget, GEngineIni);
-
-		auto FindMatchingTargetFile = [&DefaultEditorTarget](const FString& ReceiptWildcard) -> TUniquePtr<FTargetReceipt>
-		{
-			TArray<FString> AllTargetFilesWithoutPath;
-			const FString ReceiptPath = FPaths::GetPath(ReceiptWildcard);
-			IFileManager::Get().FindFiles(AllTargetFilesWithoutPath, *ReceiptWildcard, true, false);
-
-			// Used as a fallback in lieu of an exact match (only for EBuildTargetType::Editor).
-			TPair<FString, TUniquePtr<FTargetReceipt>> BestMatch;
-
-			for (const FString& TargetFileWithoutPath : AllTargetFilesWithoutPath)
-			{
-				FString TargetFile = FPaths::Combine(ReceiptPath, TargetFileWithoutPath);
-				TUniquePtr<FTargetReceipt> Receipt = MakeUnique<FTargetReceipt>();
-				if (Receipt->Read(TargetFile))
-				{
-					if (Receipt->TargetType == FApp::GetBuildTargetType() && Receipt->Configuration == FApp::GetBuildConfiguration())
-					{
-						// Priority:
-						// - If we are non-editor, first non-editor receipt
-						// - Otherwise, first that matches the running executable
-						// - Otherwise, first corresponding to DefaultEditorTarget (if specified)
-						// - Otherwise, first editor receipt
-						if (Receipt->TargetType != EBuildTargetType::Editor || Receipt->LaunchesCurrentExecutable())
-						{
-							UE_LOG(LogPluginManager, Log, TEXT("Found matching target receipt: %s"), *TargetFile);
-							return Receipt;
-						}
-
-						const bool bIsDefaultTarget = (DefaultEditorTarget.Len() != 0) && (DefaultEditorTarget == Receipt->TargetName);
-						if (bIsDefaultTarget || !BestMatch.Value)
-						{
-							BestMatch = { MoveTemp(TargetFile), MoveTemp(Receipt) };
-							// `TargetFile` and `Receipt` are now invalidated by virtue of being moved-from.
-						}
-					}
-				}
-			}
-
-			if (BestMatch.Value)
-			{
-				UE_LOG(LogPluginManager, Log, TEXT("Found target receipt: %s"), *BestMatch.Key);
-			}
-			else
-			{
-				UE_LOG(LogPluginManager, Log, TEXT("Unable to find target receipt in path: %s"), *ReceiptWildcard);
-			}
+		UE_LOG(LogPluginManager, Log, TEXT("Found target receipt: %s"), *BestMatch.Key);
+	}
+	else
+	{
+		UE_LOG(LogPluginManager, Log, TEXT("Unable to find target receipt in path: %s"), *ReceiptWildcard);
+	}
 
-			// Might be null.
-			return MoveTemp(BestMatch.Value);
-		};
+	// Might be null.
+	return MoveTemp(BestMatch.Value);
+}
 #endif // READ_TARGET_ENABLED_PLUGINS_FROM_RECEIPT
 
-#if UE_BUILD_SHIPPING
-		// const to ensure it to stays empty
-		const TSet<FString> AllowedOptionalDependencies;
-#else
-		// Set of all the plugin names that are allowed to be enabled for a plugin with optional dependencies. Only read in Editor
-		TSet<FString> AllowedOptionalDependencies;
-
+void FPluginManager::FindCompileTimePlugins(FPluginDiscoveryContext& Context)
+{
+#if !UE_BUILD_SHIPPING
 #if READ_TARGET_ENABLED_PLUGINS_FROM_RECEIPT
-		{
-			SCOPED_BOOT_TIMING("ReadTargetBuildPluginsFromReceipt");
+	{
+		SCOPED_BOOT_TIMING("ReadTargetBuildPluginsFromReceipt");
 
-			// Read the build plugins from the target file using the target receipt file. This controls which optional plugin references can be enabled.
-			auto ReadBuildPluginsFromFirstMatchingTargetFile = [&FindMatchingTargetFile, &AllowedOptionalDependencies](const TCHAR* BaseDir) -> bool
+		// Read the build plugins from the target file using the target receipt file. This controls which optional plugin references can be enabled.
+		auto ReadBuildPluginsFromFirstMatchingTargetFile = [&Context](const TCHAR* BaseDir) -> bool
 			{
 				const FString ReceiptWildcard = FTargetReceipt::GetDefaultPath(BaseDir, TEXT("*"), FPlatformProcess::GetBinariesSubdirectory(), FApp::GetBuildConfiguration(), nullptr);
 
 				TUniquePtr<FTargetReceipt> Receipt = FindMatchingTargetFile(ReceiptWildcard);
 				if (Receipt.IsValid())
 				{
-					AllowedOptionalDependencies.Append(Receipt->BuildPlugins);
+					Context.AllowedOptionalDependencies.Append(Receipt->BuildPlugins);
 					return true;
 				}
 				return false;
 			};
 
-			UE_LOG(LogPluginManager, Log, TEXT("Looking for build plugins target receipt"));
-			if (!ReadBuildPluginsFromFirstMatchingTargetFile(FPlatformMisc::ProjectDir()))
-			{
-				ReadBuildPluginsFromFirstMatchingTargetFile(FPlatformMisc::EngineDir());
-			}
-		}
-#else
+		UE_LOG(LogPluginManager, Log, TEXT("Looking for build plugins target receipt"));
+		if (!ReadBuildPluginsFromFirstMatchingTargetFile(FPlatformMisc::ProjectDir()))
 		{
-			// Configure the plugins that were enabled from the target file using defines
-			SCOPED_BOOT_TIMING("ReadTargetBuildPlugins");
-			AllowedOptionalDependencies.Append({ UBT_TARGET_BUILD_PLUGINS });
+			ReadBuildPluginsFromFirstMatchingTargetFile(FPlatformMisc::EngineDir());
 		}
+	}
+#else
+	{
+		// Configure the plugins that were enabled from the target file using defines
+		SCOPED_BOOT_TIMING("ReadTargetBuildPlugins");
+		Context.AllowedOptionalDependencies.Append({ UBT_TARGET_BUILD_PLUGINS });
+	}
 #endif // READ_TARGET_ENABLED_PLUGINS_FROM_RECEIPT
-#endif // UE_BUILD_SHIPPING
+#endif
+}
 
-		// Check which plugins have been enabled or excluded via the command line
-		{
-			SCOPED_BOOT_TIMING("ParseCmdLineForPlugins");
+// Check which plugins have been enabled or excluded via the command line
+bool FPluginManager::FindCommandLinePlugins(FPluginDiscoveryContext& Context)
+{
+	SCOPED_BOOT_TIMING("ParseCmdLineForPlugins");
 
-			auto ParsePluginsList = [this](const TCHAR* InListKey) -> TArray<FString>
+	auto ParsePluginsList = [this](const TCHAR* InListKey) -> TArray<FString>
+		{
+			TArray<FString> PluginsList;
 			{
-				TArray<FString> PluginsList;
-				{
-					FString PluginsListStr;
-					FParse::Value(FCommandLine::Get(), InListKey, PluginsListStr, false);
-					PluginsListStr.ParseIntoArray(PluginsList, TEXT(","));
-				}
+				FString PluginsListStr;
+				FParse::Value(FCommandLine::Get(), InListKey, PluginsListStr, false);
+				PluginsListStr.ParseIntoArray(PluginsList, TEXT(","));
+			}
 
-				TArray<FString> PluginListFilenames;
-				PluginsList.RemoveAll([&PluginListFilenames](const FString& ExceptPluginName)
+			TArray<FString> PluginListFilenames;
+			PluginsList.RemoveAll([&PluginListFilenames](const FString& ExceptPluginName)
 				{
 					if (ExceptPluginName.EndsWith(TEXT(".txt")))
 					{
@@ -1602,16 +1481,16 @@ bool FPluginManager::ConfigureEnabledPlugins()
 					return false;
 				});
 
-				for (const FString& PluginListFilename : PluginListFilenames) //-V1078
-				{
-					FFileHelper::LoadFileToStringWithLineVisitor(*PluginListFilename, [&PluginsList](FStringView Line)
+			for (const FString& PluginListFilename : PluginListFilenames) //-V1078
+			{
+				FFileHelper::LoadFileToStringWithLineVisitor(*PluginListFilename, [&PluginsList](FStringView Line)
 					{
 						PluginsList.Add(FString(Line));
 					});
-				}
+			}
 
-				TArray<FString> WildcardPluginsList;
-				PluginsList.RemoveAll([&WildcardPluginsList](const FString& ExceptPluginName)
+			TArray<FString> WildcardPluginsList;
+			PluginsList.RemoveAll([&WildcardPluginsList](const FString& ExceptPluginName)
 				{
 					constexpr FAsciiSet Wildcards("*?");
 					if (FAsciiSet::HasAny(ExceptPluginName, Wildcards))
@@ -1622,107 +1501,111 @@ bool FPluginManager::ConfigureEnabledPlugins()
 					return false;
 				});
 
-				if (WildcardPluginsList.Num() > 0)
+			if (WildcardPluginsList.Num() > 0)
+			{
+				for (const FString& PotentialPluginName : PluginsToConfigure)
 				{
-					for (const FString& PotentialPluginName : PluginsToConfigure)
+					bool bMatchesAnyWildcard = false;
+					for (const FString& WildcardPluginName : WildcardPluginsList) //-V1078
 					{
-						bool bMatchesAnyWildcard = false;
-						for (const FString& WildcardPluginName : WildcardPluginsList) //-V1078
+						if (PotentialPluginName.MatchesWildcard(WildcardPluginName))
 						{
-							if (PotentialPluginName.MatchesWildcard(WildcardPluginName))
-							{
-								bMatchesAnyWildcard = true;
-								break;
-							}
+							bMatchesAnyWildcard = true;
+							break;
 						}
+					}
 
-						if (bMatchesAnyWildcard)
-						{
-							PluginsList.Add(PotentialPluginName);
-						}
+					if (bMatchesAnyWildcard)
+					{
+						PluginsList.Add(PotentialPluginName);
 					}
 				}
+			}
 
-				return PluginsList;
-			};
+			return PluginsList;
+		};
 
-			// Which extra plugins should be enabled?
-			bAllPluginsEnabledViaCommandLine = FParse::Param(FCommandLine::Get(), TEXT("EnableAllPlugins"));
-			TArray<FString> ExtraPluginsToEnable;
-			if (bAllPluginsEnabledViaCommandLine)
-			{
-				ExtraPluginsToEnable = PluginsToConfigure.Array();
-			}
-			else
-			{
-				ExtraPluginsToEnable = ParsePluginsList(TEXT("EnablePlugins="));
-			}
-			if (ExtraPluginsToEnable.Num() > 0)
-			{
+	// Which extra plugins should be enabled?
+	bAllPluginsEnabledViaCommandLine = FParse::Param(FCommandLine::Get(), TEXT("EnableAllPlugins"));
+	TArray<FString> ExtraPluginsToEnable;
+	if (bAllPluginsEnabledViaCommandLine)
+	{
+		ExtraPluginsToEnable = PluginsToConfigure.Array();
+	}
+	else
+	{
+		ExtraPluginsToEnable = ParsePluginsList(TEXT("EnablePlugins="));
+	}
+	if (ExtraPluginsToEnable.Num() > 0)
+	{
 #if !UE_BUILD_SHIPPING
-				AllowedOptionalDependencies.Append(ExtraPluginsToEnable);
+		Context.AllowedOptionalDependencies.Append(ExtraPluginsToEnable);
 #endif
 
-				auto IsRestrictedPlugin = [this](const FString& PluginName)
-				{
-					if (TSharedPtr<IPlugin> PluginPtr = FindPlugin(PluginName))
-					{
-						const FString& PluginBaseDir = PluginPtr->GetBaseDir();
-						return PluginBaseDir.Contains(TEXT("/Restricted/"));
-					}
-					return true;
-				};
-
-				const TArray<FString> ExceptPlugins = ParsePluginsList(TEXT("ExceptPlugins="));
-				const bool bExceptRestrictedPlugins = FParse::Param(FCommandLine::Get(), TEXT("ExceptRestrictedPlugins"));
-				for (const FString& EnablePluginName : ExtraPluginsToEnable)
+		auto IsRestrictedPlugin = [this](const FString& PluginName)
+			{
+				if (TSharedPtr<IPlugin> PluginPtr = FindPlugin(PluginName))
 				{
-					if (!ConfiguredPluginNames.Contains(EnablePluginName) && !ExceptPlugins.Contains(EnablePluginName) && (!bExceptRestrictedPlugins || !IsRestrictedPlugin(EnablePluginName)))
-					{
-						if (!ConfigureEnabledPluginForCurrentTarget(FPluginReferenceDescriptor(EnablePluginName, true), EnabledPlugins,
-							bAllPluginsEnabledViaCommandLine ? TEXTVIEW("Commandline EnableAllPlugins") : TEXTVIEW("CommandLine EnablePlugins="), AllowedOptionalDependencies))
-						{
-							if (bAllPluginsEnabledViaCommandLine)
-							{
-								// Plugins may legitimately fail to enable when running with -EnableAllPlugins, but this shouldn't be considered a fatal error
-								continue;
-							}
-							return false;
-						}
-						ConfiguredPluginNames.Add(EnablePluginName);
-					}
+					const FString& PluginBaseDir = PluginPtr->GetBaseDir();
+					return PluginBaseDir.Contains(TEXT("/Restricted/"));
 				}
-			}
+				return true;
+			};
 
-			// Which extra plugins should be disabled?
-			TArray<FString> ExtraPluginsToDisable;
-			ExtraPluginsToDisable = ParsePluginsList(TEXT("DisablePlugins="));
-			for (const FString& DisablePluginName : ExtraPluginsToDisable)
+		const TArray<FString> ExceptPlugins = ParsePluginsList(TEXT("ExceptPlugins="));
+		const bool bExceptRestrictedPlugins = FParse::Param(FCommandLine::Get(), TEXT("ExceptRestrictedPlugins"));
+		for (const FString& EnablePluginName : ExtraPluginsToEnable)
+		{
+			if (!Context.ConfiguredPluginNames.Contains(EnablePluginName) && !ExceptPlugins.Contains(EnablePluginName) && (!bExceptRestrictedPlugins || !IsRestrictedPlugin(EnablePluginName)))
 			{
-				if (!ConfiguredPluginNames.Contains(DisablePluginName))
+				if (!ConfigureEnabledPluginForCurrentTarget(FPluginReferenceDescriptor(EnablePluginName, true), Context.EnabledPlugins,
+					bAllPluginsEnabledViaCommandLine ? TEXTVIEW("Commandline EnableAllPlugins") : TEXTVIEW("CommandLine EnablePlugins="), Context.AllowedOptionalDependencies))
 				{
-					if (!ConfigureEnabledPluginForCurrentTarget(FPluginReferenceDescriptor(DisablePluginName, false), EnabledPlugins,
-						TEXTVIEW("CommandLine DisablePlugins="), AllowedOptionalDependencies))
+					if (bAllPluginsEnabledViaCommandLine)
 					{
-						return false;
+						// Plugins may legitimately fail to enable when running with -EnableAllPlugins, but this shouldn't be considered a fatal error
+						continue;
 					}
-					ConfiguredPluginNames.Add(DisablePluginName);
+					return false;
 				}
+				Context.ConfiguredPluginNames.Add(EnablePluginName);
 			}
 		}
+	}
 
-		if (!FParse::Param(FCommandLine::Get(), TEXT("NoEnginePlugins")))
+	// Which extra plugins should be disabled?
+	TArray<FString> ExtraPluginsToDisable;
+	ExtraPluginsToDisable = ParsePluginsList(TEXT("DisablePlugins="));
+	for (const FString& DisablePluginName : ExtraPluginsToDisable)
+	{
+		if (!Context.ConfiguredPluginNames.Contains(DisablePluginName))
 		{
-			SCOPED_BOOT_TIMING("EnginePlugins");
+			if (!ConfigureEnabledPluginForCurrentTarget(FPluginReferenceDescriptor(DisablePluginName, false), Context.EnabledPlugins,
+				TEXTVIEW("CommandLine DisablePlugins="), Context.AllowedOptionalDependencies))
+			{
+				return false;
+			}
+			Context.ConfiguredPluginNames.Add(DisablePluginName);
+		}
+	}
 
-			// May be overridden later by receipt or project descriptor.
-			bool bAllowEnginePluginsEnabledByDefault = true;
+	return true;
+}
+
+bool FPluginManager::FindTargetPlugins(FPluginDiscoveryContext& Context)
+{
+	// this is a poorly named param, and probably could be just removed (this code enables engine and project plugins)
+	if (!FParse::Param(FCommandLine::Get(), TEXT("NoEnginePlugins")))
+	{
+		SCOPED_BOOT_TIMING("EnginePlugins");
+
+		// May be overridden later by receipt or project descriptor.
+		bool bAllowEnginePluginsEnabledByDefault = true;
 
 #if READ_TARGET_ENABLED_PLUGINS_FROM_RECEIPT
-			// Configure the plugins that were enabled or disabled from the target file using the target receipt file
-			auto ConfigurePluginsFromFirstMatchingTargetFile = [this, &FindMatchingTargetFile, &ConfiguredPluginNames,
-				&EnabledPlugins, &AllowedOptionalDependencies, &bAllowEnginePluginsEnabledByDefault]
-				(const TCHAR* BaseDir, bool& bOutError) -> bool
+		// Configure the plugins that were enabled or disabled from the target file using the target receipt file
+		auto ConfigurePluginsFromFirstMatchingTargetFile = [this, &Context, &bAllowEnginePluginsEnabledByDefault]
+			(const TCHAR* BaseDir, bool& bOutError) -> bool
 			{
 				const FString ReceiptWildcard = FTargetReceipt::GetDefaultPath(BaseDir, TEXT("*"), FPlatformProcess::GetBinariesSubdirectory(), FApp::GetBuildConfiguration(), nullptr);
 				FString SourceDescription = FString::Printf(TEXT("Receipt files %s"), *ReceiptWildcard);
@@ -1741,15 +1624,15 @@ bool FPluginManager::ConfigureEnabledPlugins()
 					{
 						const FString& PluginName = Pair.Key;
 						const bool bEnabled = Pair.Value;
-						if (!ConfiguredPluginNames.Contains(PluginName))
+						if (!Context.ConfiguredPluginNames.Contains(PluginName))
 						{
-							if (!ConfigureEnabledPluginForCurrentTarget(FPluginReferenceDescriptor(PluginName, bEnabled), EnabledPlugins,
-								SourceDescription, AllowedOptionalDependencies))
+							if (!ConfigureEnabledPluginForCurrentTarget(FPluginReferenceDescriptor(PluginName, bEnabled), Context.EnabledPlugins,
+								SourceDescription, Context.AllowedOptionalDependencies))
 							{
 								bOutError = true;
 								break;
 							}
-							ConfiguredPluginNames.Add(PluginName);
+							Context.ConfiguredPluginNames.Add(PluginName);
 						}
 					}
 
@@ -1759,118 +1642,126 @@ bool FPluginManager::ConfigureEnabledPlugins()
 				return false;
 			};
 
+		{
+			SCOPED_BOOT_TIMING("ConfigureTargetEnabledPluginsFromReceipt");
+			UE_LOG(LogPluginManager, Log, TEXT("Looking for enabled plugins target receipt"));
+			bool bErrorConfiguring = false;
+			if (!ConfigurePluginsFromFirstMatchingTargetFile(FPlatformMisc::ProjectDir(), bErrorConfiguring))
 			{
-				SCOPED_BOOT_TIMING("ConfigureTargetEnabledPluginsFromReceipt");
-				UE_LOG(LogPluginManager, Log, TEXT("Looking for enabled plugins target receipt"));
-				bool bErrorConfiguring = false;
-				if (!ConfigurePluginsFromFirstMatchingTargetFile(FPlatformMisc::ProjectDir(), bErrorConfiguring))
-				{
-					ConfigurePluginsFromFirstMatchingTargetFile(FPlatformMisc::EngineDir(), bErrorConfiguring);
-				}
-				if (bErrorConfiguring)
-				{
-					return false;
-				}
+				ConfigurePluginsFromFirstMatchingTargetFile(FPlatformMisc::EngineDir(), bErrorConfiguring);
 			}
-#else
+			if (bErrorConfiguring)
 			{
-				// Configure the plugins that were enabled from the target file using defines
-				SCOPED_BOOT_TIMING("ConfigureTargetEnabledPlugins");
+				return false;
+			}
+		}
+#else
+		{
+			// Configure the plugins that were enabled from the target file using defines
+			SCOPED_BOOT_TIMING("ConfigureTargetEnabledPlugins");
 
-				TArray<FString> TargetEnabledPlugins = { UBT_TARGET_ENABLED_PLUGINS };
-				for (const FString& TargetEnabledPlugin : TargetEnabledPlugins) //-V1078
+			TArray<FString> TargetEnabledPlugins = { UBT_TARGET_ENABLED_PLUGINS };
+			for (const FString& TargetEnabledPlugin : TargetEnabledPlugins) //-V1078
+			{
+				if (!Context.ConfiguredPluginNames.Contains(TargetEnabledPlugin))
 				{
-					if (!ConfiguredPluginNames.Contains(TargetEnabledPlugin))
+					if (!ConfigureEnabledPluginForCurrentTarget(FPluginReferenceDescriptor(TargetEnabledPlugin, true), Context.EnabledPlugins,
+						TEXTVIEW("UBT_TARGET_ENABLED_PLUGINS"), Context.AllowedOptionalDependencies))
 					{
-						if (!ConfigureEnabledPluginForCurrentTarget(FPluginReferenceDescriptor(TargetEnabledPlugin, true), EnabledPlugins,
-							TEXTVIEW("UBT_TARGET_ENABLED_PLUGINS"), AllowedOptionalDependencies))
-						{
-							return false;
-						}
-						ConfiguredPluginNames.Add(TargetEnabledPlugin);
+						return false;
 					}
+					Context.ConfiguredPluginNames.Add(TargetEnabledPlugin);
 				}
 			}
-			{
-				// Configure the plugins that were disabled from the target file using defines
-				SCOPED_BOOT_TIMING("ConfigureTargetDisabledPlugins");
+		}
+		{
+			// Configure the plugins that were disabled from the target file using defines
+			SCOPED_BOOT_TIMING("ConfigureTargetDisabledPlugins");
 
-				TArray<FString> TargetDisabledPlugins = { UBT_TARGET_DISABLED_PLUGINS };
-				for (const FString& TargetDisabledPlugin : TargetDisabledPlugins) //-V1078
+			TArray<FString> TargetDisabledPlugins = { UBT_TARGET_DISABLED_PLUGINS };
+			for (const FString& TargetDisabledPlugin : TargetDisabledPlugins) //-V1078
+			{
+				if (!Context.ConfiguredPluginNames.Contains(TargetDisabledPlugin))
 				{
-					if (!ConfiguredPluginNames.Contains(TargetDisabledPlugin))
+					if (!ConfigureEnabledPluginForCurrentTarget(FPluginReferenceDescriptor(TargetDisabledPlugin, false), Context.EnabledPlugins,
+						TEXTVIEW("UBT_TARGET_ENABLED_PLUGINS"), Context.AllowedOptionalDependencies))
 					{
-						if (!ConfigureEnabledPluginForCurrentTarget(FPluginReferenceDescriptor(TargetDisabledPlugin, false), EnabledPlugins,
-							TEXTVIEW("UBT_TARGET_ENABLED_PLUGINS"), AllowedOptionalDependencies))
-						{
-							return false;
-						}
-						ConfiguredPluginNames.Add(TargetDisabledPlugin);
+						return false;
 					}
+					Context.ConfiguredPluginNames.Add(TargetDisabledPlugin);
 				}
 			}
+		}
 #endif // READ_TARGET_ENABLED_PLUGINS_FROM_RECEIPT
 
-			auto ProcessPluginConfigurations =
-				[&ConfiguredPluginNames, &EnabledPlugins, &AllowedOptionalDependencies, this]
-				(const TArray<FPluginReferenceDescriptor>& PluginReferences, FStringView SourceDescription)->bool
+		auto ProcessPluginConfigurations = [&Context, this]
+			(const TArray<FPluginReferenceDescriptor>& PluginReferences, FStringView SourceDescription)->bool
 			{
 				for (const FPluginReferenceDescriptor& PluginReference : PluginReferences)
 				{
-					if (!ConfiguredPluginNames.Contains(PluginReference.Name))
+					if (!Context.ConfiguredPluginNames.Contains(PluginReference.Name))
 					{
-						if (!ConfigureEnabledPluginForCurrentTarget(PluginReference, EnabledPlugins, SourceDescription, AllowedOptionalDependencies))
+						if (!ConfigureEnabledPluginForCurrentTarget(PluginReference, Context.EnabledPlugins, SourceDescription, Context.AllowedOptionalDependencies))
 						{
 							return false;
 						}
-						ConfiguredPluginNames.Add(PluginReference.Name);
+						Context.ConfiguredPluginNames.Add(PluginReference.Name);
 					}
 				}
 				return true;
 			};
 
-			// Find all the plugin references in the project file
-			const FProjectDescriptor* ProjectDescriptor = IProjectManager::Get().GetCurrentProject();
+		// Find all the plugin references in the project file
+		const FProjectDescriptor* ProjectDescriptor = IProjectManager::Get().GetCurrentProject();
+		{
+			SCOPED_BOOT_TIMING("AddPluginReferences");
+
+			if (ProjectDescriptor != nullptr)
 			{
-				SCOPED_BOOT_TIMING("AddPluginReferences");
+				bAllowEnginePluginsEnabledByDefault = !ProjectDescriptor->bDisableEnginePluginsByDefault;
 
-				if (ProjectDescriptor != nullptr)
+				// Copy the plugin references, since we may modify the project if any plugins are missing
+				TArray<FPluginReferenceDescriptor> PluginReferences(ProjectDescriptor->Plugins);
+				if (!ProcessPluginConfigurations(PluginReferences,
+					FString::Printf(TEXT("Enabled plugins in .uproject for %s"), *ProjectDescriptor->Description)))
 				{
-					bAllowEnginePluginsEnabledByDefault = !ProjectDescriptor->bDisableEnginePluginsByDefault;
-
-					// Copy the plugin references, since we may modify the project if any plugins are missing
-					TArray<FPluginReferenceDescriptor> PluginReferences(ProjectDescriptor->Plugins);
-					if (!ProcessPluginConfigurations(PluginReferences,
-						FString::Printf(TEXT("Enabled plugins in .uproject for %s"), *ProjectDescriptor->Description)))
-					{
-						return false;
-					}
+					return false;
 				}
 			}
+		}
 
-			{
-				// Add the plugins which are enabled by default
-				SCOPED_BOOT_TIMING("AddPluginsEnabledByDefault");
+		{
+			// Add the plugins which are enabled by default
+			SCOPED_BOOT_TIMING("AddPluginsEnabledByDefault");
 
-				FString SourceDescription = FString::Printf(TEXT(".uplugin files that are enabled by default"));
-				for (const FString& PluginName : PluginsToConfigure)
+			FString SourceDescription = FString::Printf(TEXT(".uplugin files that are enabled by default"));
+			for (const FString& PluginName : PluginsToConfigure)
+			{
+				const TSharedRef<FPlugin>& Plugin = DiscoveredPluginMapUtils::FindPluginInMap_Checked(AllPlugins, PluginName);
+				if (Plugin->IsEnabledByDefault(bAllowEnginePluginsEnabledByDefault) && !Context.ConfiguredPluginNames.Contains(PluginName))
 				{
-					const TSharedRef<FPlugin>& Plugin = DiscoveredPluginMapUtils::FindPluginInMap_Checked(AllPlugins, PluginName);
-					if (Plugin->IsEnabledByDefault(bAllowEnginePluginsEnabledByDefault) && !ConfiguredPluginNames.Contains(PluginName))
+					if (!ConfigureEnabledPluginForCurrentTarget(FPluginReferenceDescriptor(PluginName, true),
+						Context.EnabledPlugins, SourceDescription, Context.AllowedOptionalDependencies))
 					{
-						if (!ConfigureEnabledPluginForCurrentTarget(FPluginReferenceDescriptor(PluginName, true),
-							EnabledPlugins, SourceDescription, AllowedOptionalDependencies))
-						{
-							return false;
-						}
-						ConfiguredPluginNames.Add(PluginName);
+						return false;
 					}
+					Context.ConfiguredPluginNames.Add(PluginName);
 				}
 			}
 		}
+	}
 
+	return true;
+}
+
+
+bool FPluginManager::FindProgramPlugins(FPluginDiscoveryContext& Context)
+{
 #if IS_PROGRAM
-		auto EnableProgramPlugin = [this, &ConfiguredPluginNames, &EnabledPlugins, &AllowedOptionalDependencies](const TCHAR* ConfigEntry, bool bOptionalPlugin) mutable
+	// Programs can also define the list of enabled plugins in ini
+	SCOPED_BOOT_TIMING("AddProgramEnabledPlugins");
+
+	auto EnableProgramPlugin = [this, &Context](const TCHAR* ConfigEntry, bool bOptionalPlugin) mutable
 		{
 			TArray<FString> ProgramPluginNames;
 			GConfig->GetArray(TEXT("Plugins"), ConfigEntry, ProgramPluginNames, GEngineIni);
@@ -1878,159 +1769,298 @@ bool FPluginManager::ConfigureEnabledPlugins()
 
 			for (const FString& PluginName : ProgramPluginNames)
 			{
-				if (!ConfiguredPluginNames.Contains(PluginName))
+				if (!Context.ConfiguredPluginNames.Contains(PluginName))
 				{
 					FPluginReferenceDescriptor PluginReference(PluginName, true);
 					PluginReference.bOptional = bOptionalPlugin;
 					if (!ConfigureEnabledPluginForCurrentTarget(PluginReference,
-						EnabledPlugins, SourceDescription, AllowedOptionalDependencies))
+						Context.EnabledPlugins, SourceDescription, Context.AllowedOptionalDependencies))
 					{
 						return false;
 					}
-					ConfiguredPluginNames.Add(PluginName);
+					Context.ConfiguredPluginNames.Add(PluginName);
 				}
 			}
 
 			return true;
 		};
 
+	if (!EnableProgramPlugin(TEXT("ProgramEnabledPlugins"), false) || !EnableProgramPlugin(TEXT("ProgramOptionalPlugins"), true))
+	{
+		return false;
+	}
+#endif
+	return true;
+}
+
+void FPluginManager::MarkEnabledPlugins(FPluginDiscoveryContext& Context)
+{
+	// Mark all the plugins as enabled
+	SCOPED_BOOT_TIMING("MarkEnabledPlugins");
+
+	// Merged modular builds act as if they were monolithic as far as plugin discovery is concerned.
+#if !IS_MONOLITHIC && !UE_MERGED_MODULES
+	TMap<FString, FString> PluginToPhysicalFile;
+
+	// If EngineBinariesRootDir is set it means that we are loading binaries from somewhere else than our basedir.
+	// We need to search for all the plugins under binaries root dir in order to find dlls etc.
+	// TODO: Maybe plugin desc inside pak files should store original relative path?
+	if (!EngineBinariesRootDir.IsEmpty())
+	{
+		TArray<FString> PhysicalFileNames;
+
+		for (const FString& EnginePluginDir : FPaths::GetExtensionDirs(EngineBinariesRootDir, TEXT("Plugins"), !GIsEditor))
 		{
-			// Programs can also define the list of enabled plugins in ini
-			SCOPED_BOOT_TIMING("AddProgramEnabledPlugins");
+			FindPluginsInDirectory(EnginePluginDir, PhysicalFileNames, IPlatformFile::GetPlatformPhysical());
+		}
 
-			if (!EnableProgramPlugin(TEXT("ProgramEnabledPlugins"), false) || !EnableProgramPlugin(TEXT("ProgramOptionalPlugins"), true))
-			{
-				return false;
-			}
+		for (const FString& ProjectPluginDir : FPaths::GetExtensionDirs(ProjectBinariesRootDir, TEXT("Plugins"), !GIsEditor))
+		{
+			FindPluginsInDirectory(ProjectPluginDir, PhysicalFileNames, IPlatformFile::GetPlatformPhysical());
 		}
-#endif
 
+		for (const FString& FileName : PhysicalFileNames)
 		{
-			// Mark all the plugins as enabled
-			SCOPED_BOOT_TIMING("MarkEnabledPlugins");
+			FString PluginName = FPaths::GetBaseFilename(FileName);
+			PluginToPhysicalFile.Add(PluginName, FileName);
+		}
+	}
+#endif // !IS_MONOLITHIC && !UE_MERGED_MODULES
+
+	for (TPair<FString, FPlugin*>& Pair : Context.EnabledPlugins)
+	{
+		FPlugin& Plugin = *Pair.Value;
 
-			// Merged modular builds act as if they were monolithic as far as plugin discovery is concerned.
 #if !IS_MONOLITHIC && !UE_MERGED_MODULES
-			TMap<FString, FString> PluginToPhysicalFile;
+		// Mount the binaries directory, and check the modules are valid.
+		// We do not run this process while using merged modular build, we'll keep acting as if the build was monolithic.
+		if (Plugin.Descriptor.Modules.Num() > 0)
+		{
+			const TCHAR* PluginFile = *Plugin.FileName;
 
-			// If EngineBinariesRootDir is set it means that we are loading binaries from somewhere else than our basedir.
-			// We need to search for all the plugins under binaries root dir in order to find dlls etc.
-			// TODO: Maybe plugin desc inside pak files should store original relative path?
-			if (!EngineBinariesRootDir.IsEmpty())
+			// If we have a overridden binary path, use that instead
+			if (FString* PhysicalFile = PluginToPhysicalFile.Find(Plugin.Name))
 			{
-				TArray<FString> PhysicalFileNames;
+				PluginFile = **PhysicalFile;
+			}
 
-				for (const FString& EnginePluginDir : FPaths::GetExtensionDirs(EngineBinariesRootDir, TEXT("Plugins"), !GIsEditor))
-				{
-					FindPluginsInDirectory(EnginePluginDir, PhysicalFileNames, IPlatformFile::GetPlatformPhysical());
-				}
+			// Mount the binaries directory
+			const FString PluginBinariesPath = FPaths::Combine(FPaths::GetPath(PluginFile), TEXT("Binaries"), FPlatformProcess::GetBinariesSubdirectory());
+			FModuleManager::Get().AddBinariesDirectory(*PluginBinariesPath, Plugin.GetLoadedFrom() == EPluginLoadedFrom::Project);
+		}
 
-				for (const FString& ProjectPluginDir : FPaths::GetExtensionDirs(ProjectBinariesRootDir, TEXT("Plugins"), !GIsEditor))
-				{
-					FindPluginsInDirectory(ProjectPluginDir, PhysicalFileNames, IPlatformFile::GetPlatformPhysical());
-				}
+		// Check the declared engine version. This is a soft requirement, so allow the user to skip over it.
+		if (!IsPluginCompatible(Plugin) && !PromptToLoadIncompatiblePlugin(Plugin))
+		{
+			UE_LOG(LogPluginManager, Display, TEXT("Skipping load of '%s'."), *Plugin.Name);
+			continue;
+		}
+#endif // !IS_MONOLITHIC && !UE_MERGED_MODULES
+		if (!Plugin.GetDescriptor().bExplicitlyLoaded && Plugin.GetDescriptor().LocalizationTargets.Num() > 0)
+		{
+			PluginLocalizationUtils::GetLocalizationPathsForPlugin(Plugin, Context.AdditionalLocResPaths);
+		}
+
+		Plugin.bEnabled = true;
+	}
+}
+
+void FPluginManager::ProcessEnabledPlugins(FPluginDiscoveryContext& Context)
+{
+	SCOPED_BOOT_TIMING("ProcessEnabledPlugins");
+
+	PluginsEnabledForTarget.Reserve(Context.EnabledPlugins.Num());
+	for (const TPair<FString, FPlugin*>& Pair : Context.EnabledPlugins)
+	{
+		if (Pair.Value->bEnabled)
+		{
+			PluginsEnabledForTarget.Add(Pair.Key);
+		}
+	}
+	// Initially this will be correct. PluginsEnabledForTarget will be updated with new mounts and this will stop including compiled plugins.
+	PluginsEnabledAtStartup.Append(PluginsEnabledForTarget);
+
+
+	// generate optimal list of plugins to process
+	TArray<TSharedRef<FPlugin>> PluginsArray;
+	for (const FString& PluginName : PluginsToConfigure)
+	{
+		TSharedRef<FPlugin> Plugin = DiscoveredPluginMapUtils::FindPluginInMap_Checked(AllPlugins, PluginName);
+
+		// only process enabled plugins
+		if (Plugin->bEnabled && !Plugin->Descriptor.bExplicitlyLoaded)
+		{
+			PluginsArray.Add(Plugin);
+		}
+	}
+
+	TArray<FString> ConfigFilesPluginsCannotOverride;
+	GConfig->GetArray(TEXT("Plugins"), TEXT("ConfigFilesPluginsCannotOverride"), ConfigFilesPluginsCannotOverride, GEngineIni);
+
+	FCriticalSection ConfigCS;
+	FCriticalSection PluginPakCS;
+
+	// Mount all the enabled plugins
+	FConfigModificationTracker ChangeTracker;
+	ChangeTracker.bTrackLoadedFiles = true;
 
-				for (const FString& FileName : PhysicalFileNames)
+	TArray<FName> PluginNames;
+	PluginNames.Reserve(PluginsArray.Num());
+
+	// walk over each plugin, and add the config files that are named for the plugin
+	// this is a separate loop so that a plugin can modify another plugin's configs in the second loop below
+	for (TSharedRef<FPlugin> PluginPtr : PluginsArray)
+	{
+		FPlugin& Plugin = *PluginPtr;
+		FName PluginName(*Plugin.Name);
+		PluginNames.Add(PluginName);
+
+		UE_LOG(LogPluginManager, Log, TEXT("Mounting %s plugin %s"), *EnumToString(Plugin.Type), *Plugin.GetName());
+		UE_LOG(LogPluginManager, Verbose, TEXT("Plugin path: %s"), *Plugin.FileName);
+
+		UE::PluginManager::Private::CheckForDeprecatedPlugin(Plugin);
+
+		// register this plugin, so it can be found right now, but also when loading other platforms, or when modifying it by other plugins
+		const bool bIncludePluginNameInBranchName = false;
+		FConfigCacheIni::RegisterPlugin(PluginName, Plugin.GetBaseDir(), Plugin.GetExtensionBaseDirs(), DynamicLayerPriority::Plugin, bIncludePluginNameInBranchName);
+
+		FConfigContext ConfigContext = FConfigContext::ReadIntoGConfig();
+		ConfigContext.ChangeTracker = &ChangeTracker;
+		ConfigContext.ConfigFileTag = *Plugin.Name;
+		ConfigContext.Load(*Plugin.Name);
+	}
+
+	// apply plugin modifications from this plugin to the everything
+	FConfigCacheIni::AddMultiplePluginsToAllBranches(PluginNames);
+
+	// walk over each plugin, modify existing ini's, mount the content, etc
+	ParallelFor(PluginsArray.Num(), [&PluginsArray, &ConfigCS, &PluginPakCS, &ConfigFilesPluginsCannotOverride, &ChangeTracker, this](int32 Index)
+		{
+			FPlugin& Plugin = *PluginsArray[Index];
+			FName PluginName(*Plugin.Name);
+
+			// Build the list of content folders
+			if (Plugin.Descriptor.bCanContainContent)
+			{
 				{
-					FString PluginName = FPaths::GetBaseFilename(FileName);
-					PluginToPhysicalFile.Add(PluginName, FileName);
+					FScopeLock Locker(&ConfigCS);
+
+					// we need to look up the section each time because other loops could add entries
+					if (FConfigFile* EngineConfigFile = GConfig->FindConfigFile(GEngineIni))
+					{
+						EngineConfigFile->AddUniqueToSection(TEXT("Core.System"), "Paths", Plugin.GetContentDir());
+					}
 				}
-			}
-#endif // !IS_MONOLITHIC && !UE_MERGED_MODULES
 
-			for (TPair<FString, FPlugin*>& Pair : EnabledPlugins)
-			{
-				FPlugin& Plugin = *Pair.Value;
+				TArray<FString>	FoundPaks;
+				FPakFileSearchVisitor PakVisitor(FoundPaks);
+				IPlatformFile& PlatformFile = FPlatformFileManager::Get().GetPlatformFile();
 
-#if !IS_MONOLITHIC && !UE_MERGED_MODULES
-				// Mount the binaries directory, and check the modules are valid.
-				// We do not run this process while using merged modular build, we'll keep acting as if the build was monolithic.
-				if (Plugin.Descriptor.Modules.Num() > 0)
+				// Pak files are loaded from <PluginName>/Content/Paks/<PlatformName>
+				if (FPlatformProperties::RequiresCookedData())
 				{
-					const TCHAR* PluginFile = *Plugin.FileName;
+					PlatformFile.IterateDirectoryRecursively(*(Plugin.GetContentDir() / TEXT("Paks") / FPlatformProperties::PlatformName()), PakVisitor);
 
-					// If we have a overridden binary path, use that instead
-					if (FString* PhysicalFile = PluginToPhysicalFile.Find(Plugin.Name))
+					for (const FString& PakPath : FoundPaks)
 					{
-						PluginFile = **PhysicalFile;
+						FScopeLock Locker(&PluginPakCS);
+						if (FCoreDelegates::MountPak.IsBound())
+						{
+							FCoreDelegates::MountPak.Execute(PakPath, 0);
+							PluginsWithPakFile.AddUnique(PluginsArray[Index]);
+						}
+						else
+						{
+							UE_LOG(LogPluginManager, Warning, TEXT("PAK file (%s) could not be mounted because MountPak is not bound"), *PakPath);
+						}
 					}
+				}
+			}
+		}, true); // @todo disable parallelism for now as it's causing hard to track problems
+}
+
+void FPluginManager::MountContentPlugins(FPluginDiscoveryContext& Context)
+{
+	// Mount enabled plugins with content
+	for (const FDiscoveredPluginMap::ElementType& PluginPair : AllPlugins)
+	{
+		const TSharedRef<FPlugin>& PluginRef = DiscoveredPluginMapUtils::ResolvePluginFromMapVal(PluginPair.Value);
+		FPlugin& Plugin = *PluginRef;
+
+		if (Plugin.IsEnabled())
+		{
+			if (Plugin.GetDescriptor().bExplicitlyLoaded)
+			{
+				continue;
+			}
 
-					// Mount the binaries directory
-					const FString PluginBinariesPath = FPaths::Combine(FPaths::GetPath(PluginFile), TEXT("Binaries"), FPlatformProcess::GetBinariesSubdirectory());
-					FModuleManager::Get().AddBinariesDirectory(*PluginBinariesPath, Plugin.GetLoadedFrom() == EPluginLoadedFrom::Project);
+			Plugin.SetIsMounted(true);
+
+			if ((Plugin.CanContainContent() || Plugin.CanContainVerse()))
+			{
+				if (NewPluginMountedEvent.IsBound())
+				{
+					NewPluginMountedEvent.Broadcast(Plugin);
 				}
 
-				// Check the declared engine version. This is a soft requirement, so allow the user to skip over it.
-				if (!IsPluginCompatible(Plugin) && !PromptToLoadIncompatiblePlugin(Plugin))
+				if (ensure(RegisterMountPointDelegate.IsBound()))
 				{
-					UE_LOG(LogPluginManager, Display, TEXT("Skipping load of '%s'."), *Plugin.Name);
-					continue;
+					FString ContentDir = Plugin.GetContentDir();
+					RegisterMountPointDelegate.Execute(Plugin.GetMountedAssetPath(), ContentDir);
 				}
-#endif // !IS_MONOLITHIC && !UE_MERGED_MODULES
-				if (!Plugin.GetDescriptor().bExplicitlyLoaded && Plugin.GetDescriptor().LocalizationTargets.Num() > 0)
+
+				if (NewPluginContentMountedEvent.IsBound())
 				{
-					PluginLocalizationUtils::GetLocalizationPathsForPlugin(Plugin, AdditionalLocResPaths);
+					NewPluginContentMountedEvent.Broadcast(Plugin);
 				}
-				
-				Plugin.bEnabled = true;
 			}
 		}
+	}
+}
 
-		PluginsEnabledForTarget.Reserve(EnabledPlugins.Num());
-		for (const TPair<FString, FPlugin*>& Pair : EnabledPlugins)
-		{
-			if (Pair.Value->bEnabled)
-			{
-				PluginsEnabledForTarget.Add(Pair.Key);
-			}
-		}
-		// Initially this will be correct. PluginsEnabledForTarget will be updated with new mounts and this will stop including compiled plugins.
-		PluginsEnabledAtStartup.Append(PluginsEnabledForTarget);
+bool FPluginManager::ConfigureEnabledPlugins()
+{
+#if (WITH_ENGINE && !IS_PROGRAM) || WITH_PLUGIN_SUPPORT
+	if(PluginsToConfigure.Num() > 0)
+	{
+		SCOPED_BOOT_TIMING("FPluginManager::ConfigureEnabledPlugins");
 
-		// If we made it here, we have all the required plugins
-		bHaveAllRequiredPlugins = true;
+		bHaveAllRequiredPlugins = false;
 
-		ProcessEnabledPlugins();
 
-		// Mount enabled plugins with content
-		for (const FDiscoveredPluginMap::ElementType& PluginPair : AllPlugins)
+		FPluginDiscoveryContext Context;
+
+		// look for compiletime and commandline plugins
+		FindCompileTimePlugins(Context);
+		if (!FindCommandLinePlugins(Context))
 		{
-			const TSharedRef<FPlugin>& PluginRef = DiscoveredPluginMapUtils::ResolvePluginFromMapVal(PluginPair.Value);
-			FPlugin& Plugin = *PluginRef;
+			return false;
+		}
 
-			if (Plugin.IsEnabled())
-			{
-				if (Plugin.GetDescriptor().bExplicitlyLoaded)
-				{
-					continue;
-				}
+		// look for the engine enabled plugins
+		if (!FindTargetPlugins(Context))
+		{
+			return false;
+		}
 
-				Plugin.SetIsMounted(true);
+		// look for program-specific plugins
+		if (!FindProgramPlugins(Context))
+		{
+			return false;
+		}
 
-				if ((Plugin.CanContainContent() || Plugin.CanContainVerse()))
-				{
-					if (NewPluginMountedEvent.IsBound())
-					{
-						NewPluginMountedEvent.Broadcast(Plugin);
-					}
+		MarkEnabledPlugins(Context);
 
-					if (ensure(RegisterMountPointDelegate.IsBound()))
-					{
-						FString ContentDir = Plugin.GetContentDir();
-						RegisterMountPointDelegate.Execute(Plugin.GetMountedAssetPath(), ContentDir);
-					}
+		// If we made it here, we have all the required plugins
+		bHaveAllRequiredPlugins = true;
 
-					if (NewPluginContentMountedEvent.IsBound())
-					{
-						NewPluginContentMountedEvent.Broadcast(Plugin);
-					}
-				}
-			}
-		}
+		ProcessEnabledPlugins(Context);
+
+		MountContentPlugins(Context);
 
-		if (AdditionalLocResPaths.Num() > 0)
+		if (Context.AdditionalLocResPaths.Num() > 0)
 		{
-			FTextLocalizationManager::Get().HandleLocalizationTargetsMounted(AdditionalLocResPaths);
+			FTextLocalizationManager::Get().HandleLocalizationTargetsMounted(Context.AdditionalLocResPaths);
 		}
 
 		PluginsToConfigure.Empty();
diff --git a/Engine/Source/Runtime/Projects/Private/PluginManager.h b/Engine/Source/Runtime/Projects/Private/PluginManager.h
--- a/Engine/Source/Runtime/Projects/Private/PluginManager.h
+++ b/Engine/Source/Runtime/Projects/Private/PluginManager.h
@@ -251,7 +251,36 @@ private:
 	    This is called when the plugin manager singleton is first accessed. */
 	void DiscoverAllPlugins();
 	
-	void ProcessEnabledPlugins();
+	/**
+	 * Helper functions just to split up the massive size of ConfigureEnabledPlugins()
+	 */
+	struct FPluginDiscoveryContext
+	{
+		// Set of all the plugins which have been enabled
+		TMap<FString, FPlugin*> EnabledPlugins;
+
+		// Keep a set of all the plugin names that have been configured. We read configuration data from different places, but only configure a plugin from the first place that it's referenced.
+		TSet<FString> ConfiguredPluginNames;
+
+		// Keep the list of newly available localization targets
+		TArray<FString> AdditionalLocResPaths;
+
+#if UE_BUILD_SHIPPING
+		// const to ensure it to stays empty
+		const TSet<FString> AllowedOptionalDependencies;
+#else
+		// Set of all the plugin names that are allowed to be enabled for a plugin with optional dependencies. Only read in Editor
+		TSet<FString> AllowedOptionalDependencies;
+#endif
+
+	};
+	void FindCompileTimePlugins(FPluginDiscoveryContext& Context);
+	bool FindCommandLinePlugins(FPluginDiscoveryContext& Context);
+	bool FindTargetPlugins(FPluginDiscoveryContext& Context);
+	bool FindProgramPlugins(FPluginDiscoveryContext& Context);
+	void MarkEnabledPlugins(FPluginDiscoveryContext& Context);
+	void ProcessEnabledPlugins(FPluginDiscoveryContext& Context);
+	void MountContentPlugins(FPluginDiscoveryContext& Context);
 
 	/** Reads all the plugin descriptors */
 	static void ReadAllPlugins(FDiscoveredPluginMap& Plugins, const TSet<FString>& ExtraSearchPaths,
