Fix to localization text keys

diff --git a/Engine/Source/Runtime/Projects/Private/PluginManifest.cpp b/Engine/Source/Runtime/Projects/Private/PluginManifest.cpp
--- a/Engine/Source/Runtime/Projects/Private/PluginManifest.cpp
+++ b/Engine/Source/Runtime/Projects/Private/PluginManifest.cpp
@@ -40,7 +40,7 @@ TOptional<FText> UE::Projects::Private::Read(UE::Json::FConstObject Object, FPlu
 	TOptional<FConstArray> JsonContents = GetArrayField(Object, TEXT("Contents"));
 	if (!JsonContents.IsSet())
 	{
-		return LOCTEXT("FailedToReadDescriptorFile", "No 'Contents' array found");
+		return LOCTEXT("PluginDescriptorMissingContentsArray", "No 'Contents' array found");
 	}
 
 	Result.Contents.Reserve(JsonContents->Size());
@@ -59,17 +59,17 @@ TOptional<FText> UE::Projects::Private::Read(UE::Json::FConstObject Object, FPlu
 
 		if (!JsonDescriptor.IsSet())
 		{
-			return FText::Format(LOCTEXT("FailedToReadDescriptorFile", "Plugin descriptor index '{0}' missing 'Descriptor' field"), Result.Contents.Num()-1);
+			return FText::Format(LOCTEXT("PluginDescriptorIndexMissingDescriptorField", "Plugin descriptor index '{0}' missing 'Descriptor' field"), Result.Contents.Num()-1);
 		}
 
 		if (!TryGetStringField(JsonEntryObject, TEXT("File"), ManifestEntry.File))
 		{
-			return FText::Format(LOCTEXT("FailedToReadDescriptorFile", "Plugin descriptor index '{0}' missing 'File' field"), Result.Contents.Num()-1);
+			return FText::Format(LOCTEXT("PluginDescriptorIndexMissingFileField", "Plugin descriptor index '{0}' missing 'File' field"), Result.Contents.Num()-1);
 		}
 
 		if (TOptional<FText> ReadError = UE::Projects::Private::Read(*JsonDescriptor, ManifestEntry.Descriptor))
 		{
-			return FText::Format(LOCTEXT("FailedToReadDescriptorFile", "Failed to parse plugin descriptor entry for file '{0}': {1}"), FText::FromString(ManifestEntry.File), *ReadError);
+			return FText::Format(LOCTEXT("PluginDescriptorEntryParseError", "Failed to parse plugin descriptor entry for file '{0}': {1}"), FText::FromString(ManifestEntry.File), *ReadError);
 		}
 	}
 
