Projects: Fix warning C4701: potentially uninitialized local variable 'FileVersionInt32' used

#rnx

diff --git a/Engine/Source/Runtime/Projects/Private/PluginDescriptor.cpp b/Engine/Source/Runtime/Projects/Private/PluginDescriptor.cpp
--- a/Engine/Source/Runtime/Projects/Private/PluginDescriptor.cpp
+++ b/Engine/Source/Runtime/Projects/Private/PluginDescriptor.cpp
@@ -178,7 +178,7 @@ bool FPluginDescriptor::Read(const FString& Text, FText& OutFailReason)
 TOptional<FText> UE::Projects::Private::Read(UE::Json::FConstObject Object, FPluginDescriptor& Result)
 {
 	// Read the file version
-	int32 FileVersionInt32;
+	int32 FileVersionInt32 = (int32)EPluginDescriptorVersion::Invalid;
 	if(!TryGetNumberField(Object, TEXT("FileVersion"), FileVersionInt32))
 	{
 		if(!TryGetNumberField(Object, TEXT("PluginFileVersion"), FileVersionInt32))
@@ -191,8 +191,8 @@ TOptional<FText> UE::Projects::Private::Read(UE::Json::FConstObject Object, FPlu
 	EPluginDescriptorVersion PluginFileVersion = (EPluginDescriptorVersion)FileVersionInt32;
 	if ((PluginFileVersion <= EPluginDescriptorVersion::Invalid) || (PluginFileVersion > EPluginDescriptorVersion::Latest))
 	{
-			const FText ReadVersionText = FText::FromString(FString::Printf(TEXT("%d"), (int32)PluginFileVersion));
-			const FText LatestVersionText = FText::FromString(FString::Printf(TEXT("%d"), (int32)EPluginDescriptorVersion::Latest));
+		const FText ReadVersionText = FText::FromString(FString::Printf(TEXT("%d"), (int32)PluginFileVersion));
+		const FText LatestVersionText = FText::FromString(FString::Printf(TEXT("%d"), (int32)EPluginDescriptorVersion::Latest));
 		return FText::Format(LOCTEXT("ProjectFileVersionTooLarge", "File appears to be in a newer version ({0}) of the file format that we can load (max version: {1})."), ReadVersionText, LatestVersionText);
 	}
 
diff --git a/Engine/Source/Runtime/Projects/Private/ProjectDescriptor.cpp b/Engine/Source/Runtime/Projects/Private/ProjectDescriptor.cpp
--- a/Engine/Source/Runtime/Projects/Private/ProjectDescriptor.cpp
+++ b/Engine/Source/Runtime/Projects/Private/ProjectDescriptor.cpp
@@ -101,7 +101,7 @@ bool FProjectDescriptor::Read(const UE::Json::FConstObject Object, const FString
 	using namespace UE::Projects::Private;
 
 	// Read the file version
-	int32 FileVersionInt32;
+	int32 FileVersionInt32 = (int32)EProjectDescriptorVersion::Invalid;
 	if(!TryGetNumberField(Object, TEXT("FileVersion"), FileVersionInt32))
 	{
 		if(!TryGetNumberField(Object, TEXT("ProjectFileVersion"), FileVersionInt32))
