LaunchEngineLoop and PluginManager: Make sure that all plugin loading phases are called in all engine configurations, even if the engine configuration does not have any code related to the loading phase, to ensure that all plugins get a chance to load themselves no matter which phase they register for.
#rnx

diff --git a/Engine/Source/Runtime/Launch/Private/LaunchEngineLoop.cpp b/Engine/Source/Runtime/Launch/Private/LaunchEngineLoop.cpp
--- a/Engine/Source/Runtime/Launch/Private/LaunchEngineLoop.cpp
+++ b/Engine/Source/Runtime/Launch/Private/LaunchEngineLoop.cpp
@@ -3269,6 +3269,7 @@ int32 FEngineLoop::PreInitPreStartupScreen(const TCHAR* CmdLine)
 				PreInitContext.SlateRenderer = SlateRenderer;
 			}
 			else
+#endif // !UE_SERVER
 			{
 				{
 					SCOPED_BOOT_TIMING("LoadModulesForProject(ELoadingPhase::PostSplashScreen)");
@@ -3279,7 +3280,6 @@ int32 FEngineLoop::PreInitPreStartupScreen(const TCHAR* CmdLine)
 					}
 				}
 			}
-#endif // !UE_SERVER
 		}
 	}
 #endif // WITH_ENGINE
@@ -3479,12 +3479,27 @@ int32 FEngineLoop::PreInitPostStartupScreen(const TCHAR* CmdLine)
 				FPlatformMisc::PlatformHandleSplashScreen(true);
 			}
 		}
-		else if (IsRunningCommandlet())
+		else
+#endif // !UE_SERVER
 		{
-			// Create the engine font services now that the Slate renderer is ready
-			FEngineFontServices::Create();
+#if !UE_SERVER
+			if (IsRunningCommandlet())
+			{
+				// Create the engine font services now that the Slate renderer is ready
+				FEngineFontServices::Create();
+			}
+#endif
+			{
+				SCOPED_BOOT_TIMING("LoadModulesForProject(ELoadingPhase::PreEarlyLoadingScreen)");
+				// Load up all modules that need to hook into the loading screen, even though there
+				// is no loading screen we still need to load the modules registered in that phase.
+				if (!IProjectManager::Get().LoadModulesForProject(ELoadingPhase::PreEarlyLoadingScreen)
+					|| !IPluginManager::Get().LoadModulesForEnabledPlugins(ELoadingPhase::PreEarlyLoadingScreen))
+				{
+					return 1;
+				}
+			}
 		}
-#endif //!UE_SERVER
 
 		//Now that our EarlyStartupScreen is finished, lets take the necessary steps to mount paks, apply .ini cvars, and open the shader libraries if we installed content we expect to handle
 		//If using a bundle manager, assume its handling all this stuff and that we don't have to do it.
diff --git a/Engine/Source/Runtime/Projects/Private/PluginManager.cpp b/Engine/Source/Runtime/Projects/Private/PluginManager.cpp
--- a/Engine/Source/Runtime/Projects/Private/PluginManager.cpp
+++ b/Engine/Source/Runtime/Projects/Private/PluginManager.cpp
@@ -2897,6 +2897,24 @@ bool FPluginManager::LoadModulesForEnabledPlugins( const ELoadingPhase::Type Loa
 		LogPluginManager, Error, TEXT("LoadModulesForEnabledPlugins called on phase %d after already being called on later phase %d."),
 		static_cast<int32>(LoadingPhase), static_cast<int32>(LastCompletedLoadingPhase));
 
+#if WITH_ENGINE
+	// Engine Startup needs to process all phases, even phases not relevant to the current engine configuration,
+	// (e.g. PostSplashScreen on Server) because some plugins might load themselves on that phase to ensure
+	// they occur before/after related code is intialized on the platforms that have code related to the phase
+	// but they still need to be loaded at SOME point on the platforms that do not have any related code. Rather
+	// than writing code to find those plugins and have them register on the next phase that occurs, just make
+	// sure that every platform processes every phase.
+	// Since the constraint enforced above means we can never go backwards in phases, this constraint means that
+	// every new phase must be the exact successor of the last completed phase (or a repeat of the last phase).
+	UE_CLOG((LastCompletedLoadingPhase == ELoadingPhase::None && LoadingPhase != ELoadingPhase::EarliestPossible)
+		|| (LastCompletedLoadingPhase != ELoadingPhase::None
+			&& static_cast<int32>(LoadingPhase) > static_cast<int32>(LastCompletedLoadingPhase) + 1),
+		LogPluginManager, Error,
+		TEXT("LoadModulesForEnabledPlugins called on phase %d after last being called on phase %d, which skipped phase %d."),
+		static_cast<int32>(LoadingPhase), static_cast<int32>(LastCompletedLoadingPhase),
+		static_cast<int32>(LastCompletedLoadingPhase) + 1);
+#endif
+
 	// We send the broadcast event each time, even if this function is called multiple times with the same phase
 	LastCompletedLoadingPhase = LoadingPhase;
 	LoadingPhaseCompleteEvent.Broadcast(LoadingPhase, bSuccess);
