Improvements to GameFeaturesSubsystem and FPluginDescriptor. Added a new hook to PluginDescriptor to allow customized loading of uplugin json data.  Updated GFS to ensure that, at least in the editor, we will try to use the PluginManager's copy of the plugin data rather than re-reading it and deserializing it from json ourselves. This is necessary to take advantage of the customized loading.

diff --git a/Engine/Plugins/Runtime/GameFeatures/Source/GameFeatures/Private/GameFeaturesSubsystem.cpp b/Engine/Plugins/Runtime/GameFeatures/Source/GameFeatures/Private/GameFeaturesSubsystem.cpp
--- a/Engine/Plugins/Runtime/GameFeatures/Source/GameFeatures/Private/GameFeaturesSubsystem.cpp
+++ b/Engine/Plugins/Runtime/GameFeatures/Source/GameFeatures/Private/GameFeaturesSubsystem.cpp
@@ -908,6 +908,18 @@ FString UGameFeaturesSubsystem::GetPluginURL_FileProtocol(const FString& PluginD
 	return Path;
 }
 
+FString UGameFeaturesSubsystem::GetPluginFilename_FileProtocol(const FString& PluginUrlFileProtocol)
+{
+	const TCHAR* Prefix = UE::GameFeatures::GameFeaturePluginProtocolPrefix(EGameFeaturePluginProtocol::File);
+	FStringView PrefixView(Prefix);
+	FString Result;
+	if (ensureAlwaysMsgf(PluginUrlFileProtocol.StartsWith(Prefix), TEXT("Unexpected protocol for %s"), *PluginUrlFileProtocol))
+	{
+		Result = PluginUrlFileProtocol.RightChop(PrefixView.Len());
+	}
+	return Result;
+}
+
 void UGameFeaturesSubsystem::GetExplanationForUnavailablePackage(const FString& UnavailablePackage, IPlugin* PluginIfFound, FStringBuilderBase& InOutExplanation)
 {
 #if WITH_EDITOR
@@ -2845,6 +2857,14 @@ bool UGameFeaturesSubsystem::GetGameFeaturePluginDetailsInternal(const FString&
 #if WITH_EDITOR
 		// In the editor we already have the plugin JSON cached
 		TSharedPtr<IPlugin> Plugin = IPluginManager::Get().FindPlugin(FPathViews::GetBaseFilename(PluginDescriptorFilename));
+		if (!Plugin.IsValid())
+		{
+			if (IPluginManager::Get().AddToPluginsList(PluginDescriptorFilename))
+			{
+				Plugin = IPluginManager::Get().FindPlugin(FPathViews::GetBaseFilename(PluginDescriptorFilename));
+			}
+		}
+
 		if (Plugin)
 		{
 			ObjectPtr = Plugin->GetDescriptorJson();
@@ -2852,20 +2872,36 @@ bool UGameFeaturesSubsystem::GetGameFeaturePluginDetailsInternal(const FString&
 		else
 #endif // WITH_EDITOR
 		{
-			// Read the file to a string
-			FString FileContents;
-			if (!FFileHelper::LoadFileToString(FileContents, *PluginDescriptorFilename))
+			bool TryWithNative = true;
+			if (FPluginDescriptor::CustomPluginDescriptorReaderDelegate.IsBound())
 			{
-				UE_LOG(LogGameFeatures, Error, TEXT("UGameFeaturesSubsystem could not load plugin descriptor. Failed to read file. File:%s Error:%d"), *PluginDescriptorFilename, FPlatformMisc::GetLastError());
-				return false;
+				TryWithNative = false;
+				FText FailReason; 
+				bool Success = FPluginDescriptor::CustomPluginDescriptorReaderDelegate.Execute(*PluginDescriptorFilename, &FailReason, ObjectPtr, TryWithNative);
+				if (!Success && !TryWithNative)
+				{
+					UE_LOG(LogGameFeatures, Error, TEXT("Failed to read plugin descriptor %s. Reason: %s"), *PluginDescriptorFilename, *FailReason.ToString());
+					return false;
+				}
 			}
 
-			// Deserialize a JSON object from the string	
-			TSharedRef<TJsonReader<>> Reader = TJsonReaderFactory<>::Create(FileContents);
-			if (!FJsonSerializer::Deserialize(Reader, ObjectPtr) || !ObjectPtr.IsValid())
+			if (TryWithNative)
 			{
-				UE_LOG(LogGameFeatures, Error, TEXT("UGameFeaturesSubsystem could not load plugin descriptor. Json invalid. File:%s. Error:%s"), *PluginDescriptorFilename, *Reader->GetErrorMessage());
-				return false;
+				// Read the file to a string
+				FString FileContents;
+				if (!FFileHelper::LoadFileToString(FileContents, *PluginDescriptorFilename))
+				{
+					UE_LOG(LogGameFeatures, Error, TEXT("UGameFeaturesSubsystem could not load plugin descriptor. Failed to read file. File:%s Error:%d"), *PluginDescriptorFilename, FPlatformMisc::GetLastError());
+					return false;
+				}
+
+				// Deserialize a JSON object from the string	
+				TSharedRef<TJsonReader<>> Reader = TJsonReaderFactory<>::Create(FileContents);
+				if (!FJsonSerializer::Deserialize(Reader, ObjectPtr) || !ObjectPtr.IsValid())
+				{
+					UE_LOG(LogGameFeatures, Error, TEXT("UGameFeaturesSubsystem could not load plugin descriptor. Json invalid. File:%s. Error:%s"), *PluginDescriptorFilename, *Reader->GetErrorMessage());
+					return false;
+				}
 			}
 		}
 	}
diff --git a/Engine/Plugins/Runtime/GameFeatures/Source/GameFeatures/Public/GameFeaturesSubsystem.h b/Engine/Plugins/Runtime/GameFeatures/Source/GameFeatures/Public/GameFeaturesSubsystem.h
--- a/Engine/Plugins/Runtime/GameFeatures/Source/GameFeatures/Public/GameFeaturesSubsystem.h
+++ b/Engine/Plugins/Runtime/GameFeatures/Source/GameFeatures/Public/GameFeaturesSubsystem.h
@@ -505,6 +505,9 @@ public:
 	/** Construct a 'file:' Plugin URL using from the PluginDescriptorPath */
 	static UE_API FString GetPluginURL_FileProtocol(const FString& PluginDescriptorPath);
 	static UE_API FString GetPluginURL_FileProtocol(const FString& PluginDescriptorPath, TArrayView<const TPair<FString, FString>> AdditionalOptions);
+	
+	/** Get the file path portion of a file protocol URL*/
+	static UE_API FString GetPluginFilename_FileProtocol(const FString& PluginUrlFileProtocol);
 
 	/** Construct a 'installbundle:' Plugin URL using from the PluginName and required install bundles */
 	static UE_API FString GetPluginURL_InstallBundleProtocol(const FString& PluginName, TArrayView<const FString> BundleNames);
diff --git a/Engine/Source/Runtime/Projects/Private/PluginDescriptor.cpp b/Engine/Source/Runtime/Projects/Private/PluginDescriptor.cpp
--- a/Engine/Source/Runtime/Projects/Private/PluginDescriptor.cpp
+++ b/Engine/Source/Runtime/Projects/Private/PluginDescriptor.cpp
@@ -11,9 +11,9 @@
 
 #define LOCTEXT_NAMESPACE "PluginDescriptor"
 
-namespace PluginDescriptor
+namespace UE::PluginDescriptor::Internal
 {
-	bool ReadFile(const TCHAR* FileName, FString& Text, FText* OutFailReason = nullptr)
+	bool ReadFile(const TCHAR* FileName, FString& Text, FText* OutFailReason/* = nullptr*/)
 	{
 		if (!FFileHelper::LoadFileToString(Text, FileName))
 		{
@@ -26,7 +26,7 @@ namespace PluginDescriptor
 		return true;
 	}
 
-	bool WriteFile(const TCHAR* FileName, const FString& Text, FText* OutFailReason = nullptr)
+	bool WriteFile(const TCHAR* FileName, const FString& Text, FText* OutFailReason/* = nullptr*/)
 	{
 		if (!FFileHelper::SaveStringToFile(Text, FileName))
 		{
@@ -39,7 +39,7 @@ namespace PluginDescriptor
 		return true;
 	}
 
-	TSharedPtr<FJsonObject> DeserializeJson(const FString& Text, FText* OutFailReason = nullptr)
+	TSharedPtr<FJsonObject> DeserializeJson(const FString& Text, FText* OutFailReason/* = nullptr*/)
 	{
 		TSharedPtr<FJsonObject> JsonObject;
 		TSharedRef<TJsonReader<>> Reader = TJsonReaderFactory<>::Create(Text);
@@ -104,11 +104,20 @@ bool FPluginDescriptor::Load(const TCHAR* FileName, FText* OutFailReason /*= nul
 	AdditionalFieldsToRemove.Reset();
 #endif // WITH_EDITOR
 
-	FString Text;
-	if (PluginDescriptor::ReadFile(FileName, Text, OutFailReason))
+	bool GotText = false;
+	if (CustomPluginDescriptorReaderDelegate.IsBound())
+	{
+		return ReadWithCustomPluginDescriptorReader(FileName, OutFailReason);
+	}
+	else
 	{
-		return Read(Text, OutFailReason);
+		FString Text;
+		if (UE::PluginDescriptor::Internal::ReadFile(FileName, Text, OutFailReason))
+		{
+			return Read(Text, OutFailReason);
+		}
 	}
+
 	return false;
 }
 
@@ -122,6 +131,73 @@ bool FPluginDescriptor::Load(const FString& FileName, FText& OutFailReason)
 	return Load(*FileName, &OutFailReason);
 }
 
+bool FPluginDescriptor::ReadWithCustomPluginDescriptorReader(const TCHAR* FileName, FText* OutFailReason /*= nullptr*/)
+{
+	using namespace UE::Json;
+
+	TSharedPtr<FJsonObject> JsonObject = MakeShared<FJsonObject>();
+	bool RetryWithNative = true;
+	if (CustomPluginDescriptorReaderDelegate.IsBound())
+	{
+		RetryWithNative = false;
+		bool Success = CustomPluginDescriptorReaderDelegate.Execute(FileName, OutFailReason, JsonObject, RetryWithNative);
+		if (!Success && !RetryWithNative)
+		{
+			return false;
+		}
+	}
+
+	if (!JsonObject.IsValid())
+	{
+		if (RetryWithNative)
+		{
+			FString Text;
+			if (UE::PluginDescriptor::Internal::ReadFile(FileName, Text, OutFailReason))
+			{
+				return Read(Text, OutFailReason);
+			}
+		}
+		return false;
+	}
+
+
+#if WITH_EDITOR
+	CachedJson.Reset();
+	AdditionalFieldsToWrite.Reset();
+	AdditionalFieldsToRemove.Reset();
+#endif // WITH_EDITOR
+
+	TOptional<FDocument> RapidJsonDocument = ConvertSharedJsonToRapidJsonDocument(*JsonObject);
+	if (!RapidJsonDocument.IsSet())
+	{
+		*OutFailReason = FText::Format(LOCTEXT("InvalidDescriptorFile", "Internal error or malformed plugin descriptor file. {0}"), FText::FromString(FileName));
+		return false;
+	}
+
+	TOptional<FConstObject> RootObject = GetRootObject(RapidJsonDocument.GetValue());
+	if (RootObject.IsSet())
+	{
+		// Parse it as a plug-in descriptor
+		if (TOptional<FText> ReadError = UE::Projects::Private::Read(*RootObject, *this))
+		{
+			if (OutFailReason)
+			{
+				*OutFailReason = MoveTemp(*ReadError);
+			}
+			return false;
+		}
+
+
+#if WITH_EDITOR
+		CachedJson = JsonObject;
+		AdditionalFieldsToWrite.Reset();
+		AdditionalFieldsToRemove.Reset();
+#endif // WITH_EDITOR
+		return true;
+	}
+	return false;
+}
+
 bool FPluginDescriptor::Read(const FString& Text, FText* OutFailReason /*= nullptr*/)
 {
 	using namespace UE::Json;
@@ -138,7 +214,7 @@ bool FPluginDescriptor::Read(const FString& Text, FText* OutFailReason /*= nullp
 	if (ParseResult.HasError())
 	{
 		if (OutFailReason)
-	{
+		{
 			*OutFailReason = FText::Format(LOCTEXT("FailedToReadDescriptorFile", "Failed to read file. {0}"), FText::FromString(ParseResult.GetError().CreateMessage(Text)));
 		}
 		return false;
@@ -344,7 +420,7 @@ bool FPluginDescriptor::Save(const TCHAR* FileName, FText* OutFailReason /*= nul
 	Write(Text);
 
 	// Save it to a file
-	return PluginDescriptor::WriteFile(FileName, Text, OutFailReason);
+	return UE::PluginDescriptor::Internal::WriteFile(FileName, Text, OutFailReason);
 }
 
 bool FPluginDescriptor::Save(const FString& FileName, FText* OutFailReason /*= nullptr*/) const
@@ -608,12 +684,12 @@ bool FPluginDescriptor::UpdatePluginFile(const FString& FileName, FText* OutFail
 		// Plugin file exists so we need to read it and update it.
 
 		FString JsonText;
-		if (!PluginDescriptor::ReadFile(*FileName, JsonText, OutFailReason))
+		if (!UE::PluginDescriptor::Internal::ReadFile(*FileName, JsonText, OutFailReason))
 		{
 			return false;
 		}
 
-		TSharedPtr<FJsonObject> JsonObject = PluginDescriptor::DeserializeJson(JsonText, OutFailReason);
+		TSharedPtr<FJsonObject> JsonObject = UE::PluginDescriptor::Internal::DeserializeJson(JsonText, OutFailReason);
 		if (!JsonObject.IsValid())
 		{
 			return false;
@@ -636,7 +712,7 @@ bool FPluginDescriptor::UpdatePluginFile(const FString& FileName, FText* OutFail
 #if WITH_EDITOR
 		CachedJson = JsonObject;
 #endif
-		return PluginDescriptor::WriteFile(*FileName, JsonText, OutFailReason);
+		return UE::PluginDescriptor::Internal::WriteFile(*FileName, JsonText, OutFailReason);
 	}
 	else
 	{
@@ -662,4 +738,6 @@ bool FPluginDescriptor::SupportsTargetPlatform(const FString& Platform) const
 	}
 }
 
+
+FPluginDescriptor::FPluginDescriptorReaderDelegate FPluginDescriptor::CustomPluginDescriptorReaderDelegate;
 #undef LOCTEXT_NAMESPACE
diff --git a/Engine/Source/Runtime/Projects/Public/PluginDescriptor.h b/Engine/Source/Runtime/Projects/Public/PluginDescriptor.h
--- a/Engine/Source/Runtime/Projects/Public/PluginDescriptor.h
+++ b/Engine/Source/Runtime/Projects/Public/PluginDescriptor.h
@@ -9,6 +9,7 @@
 #include "CustomBuildSteps.h"
 #include "HAL/Platform.h"
 #include "LocalizationDescriptor.h"
+#include "Misc/CoreDelegates.h"
 #include "Misc/Optional.h"
 #include "ModuleDescriptor.h"
 #include "PluginDisallowedDescriptor.h"
@@ -191,6 +192,9 @@ struct FPluginDescriptor
 	/** Loads the descriptor from the given file. */
 	PROJECTS_API bool Load(const FString& FileName, FText& OutFailReason);
 
+	/* Reads the descriptor using the global CustomPluginDescriptorReaderDelegate */
+	PROJECTS_API bool ReadWithCustomPluginDescriptorReader(const TCHAR* FileName, FText* OutFailReason = nullptr);
+
 	/** Reads the descriptor from the given string */
 	PROJECTS_API bool Read(const FString& Text, FText* OutFailReason = nullptr);
 
@@ -235,5 +239,14 @@ struct FPluginDescriptor
 
 	/** Determines whether the plugin supports the given platform */
 	PROJECTS_API bool SupportsTargetPlatform(const FString& Platform) const;
+
+	// Custom hook for providing a plugin descriptor (as an FJsonObject) based on a filename
+	DECLARE_DELEGATE_RetVal_FourParams(bool, FPluginDescriptorReaderDelegate, const TCHAR*/*Filename*/, FText*/*OutFailReason*/,  TSharedPtr<FJsonObject>&/*OutJson*/, bool&/*OutRetryWithNative*/);
+	PROJECTS_API static FPluginDescriptorReaderDelegate CustomPluginDescriptorReaderDelegate;
 };
 
+namespace UE::PluginDescriptor::Internal
+{
+	PROJECTS_API bool ReadFile(const TCHAR* FileName, FString& Text, FText* OutFailReason = nullptr);
+	PROJECTS_API TSharedPtr<FJsonObject> DeserializeJson(const FString& Text, FText* OutFailReason = nullptr);
+}
