Build warning fix

#rnx

diff --git a/Engine/Source/Runtime/Projects/Private/PluginDescriptor.cpp b/Engine/Source/Runtime/Projects/Private/PluginDescriptor.cpp
--- a/Engine/Source/Runtime/Projects/Private/PluginDescriptor.cpp
+++ b/Engine/Source/Runtime/Projects/Private/PluginDescriptor.cpp
@@ -170,7 +170,10 @@ bool FPluginDescriptor::ReadWithCustomPluginDescriptorReader(const TCHAR* FileNa
 	TOptional<FDocument> RapidJsonDocument = ConvertSharedJsonToRapidJsonDocument(*JsonObject);
 	if (!RapidJsonDocument.IsSet())
 	{
-		*OutFailReason = FText::Format(LOCTEXT("InvalidDescriptorFile", "Internal error or malformed plugin descriptor file. {0}"), FText::FromString(FileName));
+		if (OutFailReason != nullptr)
+		{
+			*OutFailReason = FText::Format(LOCTEXT("InvalidDescriptorFile", "Internal error or malformed plugin descriptor file. {0}"), FText::FromString(FileName));
+		}
 		return false;
 	}
 
