Improving safety of FPluginDescriptor::ReadWithCustomPluginDescriptorReader

#rnx

diff --git a/Engine/Source/Runtime/Projects/Private/PluginDescriptor.cpp b/Engine/Source/Runtime/Projects/Private/PluginDescriptor.cpp
--- a/Engine/Source/Runtime/Projects/Private/PluginDescriptor.cpp
+++ b/Engine/Source/Runtime/Projects/Private/PluginDescriptor.cpp
@@ -147,10 +147,12 @@ bool FPluginDescriptor::ReadWithCustomPluginDescriptorReader(const TCHAR* FileNa
 		}
 	}
 
-	if (!JsonObject.IsValid())
+	if (!JsonObject.IsValid() || RetryWithNative)
 	{
 		if (RetryWithNative)
 		{
+			// Reset the json container
+			JsonObject = MakeShared<FJsonObject>();
 			FString Text;
 			if (UE::PluginDescriptor::Internal::ReadFile(FileName, Text, OutFailReason))
 			{
