Small fix to FPluginManager::ConfigureEnabledPlugins performance

diff --git a/Engine/Source/Runtime/Projects/Private/PluginManager.cpp b/Engine/Source/Runtime/Projects/Private/PluginManager.cpp
--- a/Engine/Source/Runtime/Projects/Private/PluginManager.cpp
+++ b/Engine/Source/Runtime/Projects/Private/PluginManager.cpp
@@ -1660,17 +1660,16 @@ bool FPluginManager::FindTargetPlugins(FPluginDiscoveryContext& Context)
 			// Configure the plugins that were enabled from the target file using defines
 			SCOPED_BOOT_TIMING("ConfigureTargetEnabledPlugins");
 
-			TArray<FString> TargetEnabledPlugins = { UBT_TARGET_ENABLED_PLUGINS };
-			for (const FString& TargetEnabledPlugin : TargetEnabledPlugins) //-V1078
+			for (const FStringView& TargetEnabledPlugin : std::initializer_list<FStringView>{UBT_TARGET_ENABLED_PLUGINS}) //-V1078
 			{
-				if (!Context.ConfiguredPluginNames.Contains(TargetEnabledPlugin))
+				if (!Context.ConfiguredPluginNames.ContainsByHash(GetTypeHash(TargetEnabledPlugin), TargetEnabledPlugin))
 				{
-					if (!ConfigureEnabledPluginForCurrentTarget(FPluginReferenceDescriptor(TargetEnabledPlugin, true), Context.EnabledPlugins,
+					if (!ConfigureEnabledPluginForCurrentTarget(FPluginReferenceDescriptor(FString{TargetEnabledPlugin}, true), Context.EnabledPlugins,
 						TEXTVIEW("UBT_TARGET_ENABLED_PLUGINS"), Context.AllowedOptionalDependencies))
 					{
 						return false;
 					}
-					Context.ConfiguredPluginNames.Add(TargetEnabledPlugin);
+					Context.ConfiguredPluginNames.Add(FString{TargetEnabledPlugin});
 				}
 			}
 		}
@@ -1678,17 +1677,16 @@ bool FPluginManager::FindTargetPlugins(FPluginDiscoveryContext& Context)
 			// Configure the plugins that were disabled from the target file using defines
 			SCOPED_BOOT_TIMING("ConfigureTargetDisabledPlugins");
 
-			TArray<FString> TargetDisabledPlugins = { UBT_TARGET_DISABLED_PLUGINS };
-			for (const FString& TargetDisabledPlugin : TargetDisabledPlugins) //-V1078
+			for (const FStringView& TargetDisabledPlugin : std::initializer_list<FStringView>{UBT_TARGET_DISABLED_PLUGINS}) //-V1078
 			{
-				if (!Context.ConfiguredPluginNames.Contains(TargetDisabledPlugin))
+				if (!Context.ConfiguredPluginNames.ContainsByHash(GetTypeHash(TargetDisabledPlugin), TargetDisabledPlugin))
 				{
-					if (!ConfigureEnabledPluginForCurrentTarget(FPluginReferenceDescriptor(TargetDisabledPlugin, false), Context.EnabledPlugins,
+					if (!ConfigureEnabledPluginForCurrentTarget(FPluginReferenceDescriptor(FString{TargetDisabledPlugin}, false), Context.EnabledPlugins,
 						TEXTVIEW("UBT_TARGET_ENABLED_PLUGINS"), Context.AllowedOptionalDependencies))
 					{
 						return false;
 					}
-					Context.ConfiguredPluginNames.Add(TargetDisabledPlugin);
+					Context.ConfiguredPluginNames.Add(FString{TargetDisabledPlugin});
 				}
 			}
 		}
diff --git a/Engine/Source/Runtime/Projects/Private/PluginReferenceDescriptor.cpp b/Engine/Source/Runtime/Projects/Private/PluginReferenceDescriptor.cpp
--- a/Engine/Source/Runtime/Projects/Private/PluginReferenceDescriptor.cpp
+++ b/Engine/Source/Runtime/Projects/Private/PluginReferenceDescriptor.cpp
@@ -29,8 +29,8 @@ namespace PluginReferenceDescriptor
 	}
 }
 
-FPluginReferenceDescriptor::FPluginReferenceDescriptor( const FString& InName, bool bInEnabled )
-	: Name(InName)
+FPluginReferenceDescriptor::FPluginReferenceDescriptor(FString InName, bool bInEnabled)
+	: Name(MoveTemp(InName))
 	, bEnabled(bInEnabled)
 	, bOptional(false)
 	, bActivate(false)
diff --git a/Engine/Source/Runtime/Projects/Public/PluginReferenceDescriptor.h b/Engine/Source/Runtime/Projects/Public/PluginReferenceDescriptor.h
--- a/Engine/Source/Runtime/Projects/Public/PluginReferenceDescriptor.h
+++ b/Engine/Source/Runtime/Projects/Public/PluginReferenceDescriptor.h
@@ -79,7 +79,7 @@ struct FPluginReferenceDescriptor
 #endif //if WITH_EDITOR
 
 	/** Constructor */
-	PROJECTS_API FPluginReferenceDescriptor(const FString& InName = TEXT(""), bool bInEnabled = false);
+	PROJECTS_API FPluginReferenceDescriptor(FString InName = TEXT(""), bool bInEnabled = false);
 
 	/** Determines whether the plugin is enabled for the given platform */
 	PROJECTS_API bool IsEnabledForPlatform(const FString& Platform) const;
