Add Support for the copy stage handler to know the engine config hierachy. This can allow users to make ini configuration changes based upon the platform being staged.
Add ability to support the maximum loading phase an explicitly mounted plugin should load the modules for. Useful for when you want to load a plugin in the earliest possible part but the plugin dependencies should be loaded later.

diff --git a/Engine/Source/Programs/AutomationTool/AutomationUtils/CustomStageCopyHandler.cs b/Engine/Source/Programs/AutomationTool/AutomationUtils/CustomStageCopyHandler.cs
--- a/Engine/Source/Programs/AutomationTool/AutomationUtils/CustomStageCopyHandler.cs
+++ b/Engine/Source/Programs/AutomationTool/AutomationUtils/CustomStageCopyHandler.cs
@@ -71,7 +71,17 @@ public abstract class CustomStageCopyHandler
 		return Handler;
 	}
 
+	public virtual void Init(ConfigHierarchy InPlatformEngineConfig)
+	{
+		PlatformEngineConfig = InPlatformEngineConfig;
+	}
+
 	private static Dictionary<string, Type> HandlerRegistry = null;
+
+	/// <summary>
+	/// Platform engine config hierarchy for the platform being staged.
+	/// </summary>
+	protected ConfigHierarchy PlatformEngineConfig = null;
 }
 
 /// <summary>
diff --git a/Engine/Source/Programs/AutomationTool/Scripts/CopyBuildToStagingDirectory.Automation.cs b/Engine/Source/Programs/AutomationTool/Scripts/CopyBuildToStagingDirectory.Automation.cs
--- a/Engine/Source/Programs/AutomationTool/Scripts/CopyBuildToStagingDirectory.Automation.cs
+++ b/Engine/Source/Programs/AutomationTool/Scripts/CopyBuildToStagingDirectory.Automation.cs
@@ -6398,6 +6398,7 @@ namespace AutomationScripts
 				try
 				{
 					SC.CustomStageCopyHandler = CustomStageCopyHandler.Create(CustomStageCopyHandlerName);
+					SC.CustomStageCopyHandler.Init(ConfigCache.ReadHierarchy(ConfigHierarchyType.Engine, DirectoryReference.FromFile(Params.RawProjectPath), SC.StageTargetPlatform.IniPlatformType));
 				}
 				catch
 				{
diff --git a/Engine/Source/Runtime/Projects/Private/PluginManager.cpp b/Engine/Source/Runtime/Projects/Private/PluginManager.cpp
--- a/Engine/Source/Runtime/Projects/Private/PluginManager.cpp
+++ b/Engine/Source/Runtime/Projects/Private/PluginManager.cpp
@@ -3289,14 +3289,14 @@ void FPluginManager::MountNewlyCreatedPlugin(const FString& PluginName)
 	}
 }
 
-bool FPluginManager::MountExplicitlyLoadedPlugin(const FString& PluginName)
+bool FPluginManager::MountExplicitlyLoadedPlugin(const FString& PluginName, ELoadingPhase::Type MaxLoadingPhase)
 {
 	bool bSuccess = false;
 
 	TSharedPtr<FPlugin> Plugin = FindPluginInstance(PluginName);
 	if (Plugin.IsValid() && Plugin->Descriptor.bExplicitlyLoaded)
 	{
-		MountPluginFromExternalSource(Plugin.ToSharedRef());
+		MountPluginFromExternalSource(Plugin.ToSharedRef(), MaxLoadingPhase);
 		bSuccess = true;
 	}
 
@@ -3510,7 +3510,7 @@ bool FPluginManager::TryMountExplicitlyLoadedPluginVersion(TSharedRef<FPlugin>*
 	return bSuccess;
 }
 
-void FPluginManager::MountPluginFromExternalSource(const TSharedRef<FPlugin>& Plugin)
+void FPluginManager::MountPluginFromExternalSource(const TSharedRef<FPlugin>& Plugin, ELoadingPhase::Type MaxLoadingPhase)
 {
 	TRACE_CPUPROFILER_EVENT_SCOPE(MountPluginFromExternalSource);
 	if (GWarn)
@@ -3562,7 +3562,7 @@ void FPluginManager::MountPluginFromExternalSource(const TSharedRef<FPlugin>& Pl
 		FModuleManager::Get().AddBinariesDirectory(*PluginBinariesPath, Plugin->GetLoadedFrom() == EPluginLoadedFrom::Project);
 
 		// Load all the plugin modules
-		for (ELoadingPhase::Type LoadingPhase = (ELoadingPhase::Type)0; LoadingPhase < ELoadingPhase::Max; LoadingPhase = (ELoadingPhase::Type)(LoadingPhase + 1))
+		for (ELoadingPhase::Type LoadingPhase = (ELoadingPhase::Type)0; LoadingPhase < MaxLoadingPhase; LoadingPhase = (ELoadingPhase::Type)(LoadingPhase + 1))
 		{
 			if (LoadingPhase != ELoadingPhase::None)
 			{
diff --git a/Engine/Source/Runtime/Projects/Private/PluginManager.h b/Engine/Source/Runtime/Projects/Private/PluginManager.h
--- a/Engine/Source/Runtime/Projects/Private/PluginManager.h
+++ b/Engine/Source/Runtime/Projects/Private/PluginManager.h
@@ -222,7 +222,7 @@ public:
 	virtual FNewPluginMountedEvent& OnPluginEdited() override;
 	virtual FNewPluginMountedEvent& OnPluginUnmounted() override;
 	virtual void MountNewlyCreatedPlugin(const FString& PluginName) override;
-	virtual bool MountExplicitlyLoadedPlugin(const FString& PluginName) override;
+	virtual bool MountExplicitlyLoadedPlugin(const FString& PluginName, ELoadingPhase::Type MaxLoadingPhase = ELoadingPhase::Max) override;
 	virtual bool MountExplicitlyLoadedPlugin_FromFileName(const FString& PluginFileName) override;
 	virtual bool MountExplicitlyLoadedPlugin_FromDescriptor(const FPluginReferenceDescriptor& PluginDescriptor) override;
 	virtual bool MountExplicitlyLoadedPluginLocalizationData(const FString& PluginName, TUniqueFunction<void(bool bLoadedLocalization, const FString& PluginName)>&& Notification = nullptr) override;
@@ -354,7 +354,7 @@ private:
 	bool TryMountExplicitlyLoadedPluginVersion(TSharedRef<FPlugin>* AllPlugins_PluginPtr);
 
 	/** Mounts a plugin that was requested to be mounted from external code (either by MountNewlyCreatedPlugin or MountExplicitlyLoadedPlugin) */
-	void MountPluginFromExternalSource(const TSharedRef<FPlugin>& Plugin);
+	void MountPluginFromExternalSource(const TSharedRef<FPlugin>& Plugin, ELoadingPhase::Type MaxLoadingPhase = ELoadingPhase::Max);
 
 #if WITH_EDITOR
 	void AddToModuleNameToPluginMap(const TSharedRef<FPlugin>& Plugin);
diff --git a/Engine/Source/Runtime/Projects/Public/Interfaces/IPluginManager.h b/Engine/Source/Runtime/Projects/Public/Interfaces/IPluginManager.h
--- a/Engine/Source/Runtime/Projects/Public/Interfaces/IPluginManager.h
+++ b/Engine/Source/Runtime/Projects/Public/Interfaces/IPluginManager.h
@@ -542,7 +542,7 @@ public:
 	 * 
 	 * @note Call MountExplicitlyLoadedPluginLocalizationData if you also want to load any localization data for this plugin.
 	 */
-	virtual bool MountExplicitlyLoadedPlugin(const FString& PluginName) = 0;
+	virtual bool MountExplicitlyLoadedPlugin(const FString& PluginName, ELoadingPhase::Type MaxLoadingPhase = ELoadingPhase::Max) = 0;
 	virtual bool MountExplicitlyLoadedPlugin_FromFileName(const FString& PluginFileName) = 0;
 	virtual bool MountExplicitlyLoadedPlugin_FromDescriptor(const FPluginReferenceDescriptor& PluginDescriptor) = 0;
 
