[Project]
* Fixed so UBT_TARGET_ENABLED_PLUGINS and UBT_TARGET_DISABLED_PLUGINS

diff --git a/Engine/Source/Runtime/Projects/Projects.Build.cs b/Engine/Source/Runtime/Projects/Projects.Build.cs
--- a/Engine/Source/Runtime/Projects/Projects.Build.cs
+++ b/Engine/Source/Runtime/Projects/Projects.Build.cs
@@ -37,15 +37,17 @@ namespace UnrealBuildTool.Rules
 				List<string> EnabledPluginStrings = new List<string>();
 				foreach(string EnablePlugin in Target.EnablePlugins)
 				{
-					EnabledPluginStrings.Add(String.Format("TEXT(\"{0}\")", EnablePlugin));
+					EnabledPluginStrings.Add($"TEXT(\"{EnablePlugin}\")");
 				}
+				EnabledPluginStrings.Sort();
 				PrivateDefinitions.Add(String.Format("UBT_TARGET_ENABLED_PLUGINS={0}", String.Join(", ", EnabledPluginStrings)));
 
 				List<string> DisabledPluginStrings = new List<string>();
 				foreach(string DisablePlugin in Target.DisablePlugins)
 				{
-					DisabledPluginStrings.Add(String.Format("TEXT(\"{0}\")", DisablePlugin));
+					DisabledPluginStrings.Add($"TEXT(\"{DisablePlugin}\")");
 				}
+				DisabledPluginStrings.Sort();
 				PrivateDefinitions.Add(String.Format("UBT_TARGET_DISABLED_PLUGINS={0}", String.Join(", ", DisabledPluginStrings)));
 			}
 
