GitHub 14173 : Removing force-setting plugins enabled to false during unload

#14173

diff --git a/Engine/Source/Runtime/Projects/Private/PluginManager.cpp b/Engine/Source/Runtime/Projects/Private/PluginManager.cpp
--- a/Engine/Source/Runtime/Projects/Private/PluginManager.cpp
+++ b/Engine/Source/Runtime/Projects/Private/PluginManager.cpp
@@ -3706,7 +3706,6 @@ bool FPluginManager::UnmountExplicitlyLoadedPlugin(const FString& PluginName, FT
 	}
 
 	Plugin->SetIsMounted(false);
-	Plugin->bEnabled = false;
 
 	return true;
 }
