RapidJsonPluginLoadingTests: Skip GetAdditionalPluginDirectories test for BlankProgram that is built without WITH_PLUGIN_SUPPORT.

#rnx

#Synced-CL 49808502

diff --git a/Engine/Source/Runtime/Projects/Private/Tests/RapidJsonPluginLoadingTests.cpp b/Engine/Source/Runtime/Projects/Private/Tests/RapidJsonPluginLoadingTests.cpp
--- a/Engine/Source/Runtime/Projects/Private/Tests/RapidJsonPluginLoadingTests.cpp
+++ b/Engine/Source/Runtime/Projects/Private/Tests/RapidJsonPluginLoadingTests.cpp
@@ -430,7 +430,10 @@ TEST_CASE_NAMED(FProjectsLoadingProjectDescriptorTests, "System::Engine::Project
 	REQUIRE(Descriptor.Description == TEXT("Description goes here"));
 	REQUIRE(Descriptor.bIsEnterpriseProject == true);
 	REQUIRE(Descriptor.bDisableEnginePluginsByDefault == true);
-	REQUIRE(Descriptor.GetAdditionalPluginDirectories().Num() == 1);
+	if (FPlatformProperties::RequiresCookedData() || WITH_EDITOR || (IS_PROGRAM && WITH_PLUGIN_SUPPORT))
+	{
+		REQUIRE(Descriptor.GetAdditionalPluginDirectories().Num() == 1);
+	}
 #if WITH_EDITOR	
 	REQUIRE(Descriptor.GetAdditionalRootDirectories().Num() == 1);
 #endif
