Fix for regression issue in engine 5.7, where old DOM json parser had case insensitive keys in the project files.

#Synced-CL 49848080

diff --git a/Engine/Source/Runtime/Projects/Private/RapidJsonPluginLoading.cpp b/Engine/Source/Runtime/Projects/Private/RapidJsonPluginLoading.cpp
--- a/Engine/Source/Runtime/Projects/Private/RapidJsonPluginLoading.cpp
+++ b/Engine/Source/Runtime/Projects/Private/RapidJsonPluginLoading.cpp
@@ -11,9 +11,30 @@ namespace UE {
 namespace Projects {
 namespace Private {
 
+// the DOM json parser is case insensitive on keys names
+Json::FValue::ConstMemberIterator FindMemberCaseInsensitive(Json::FConstObject Object, const TCHAR* FieldName)
+{
+	const int32 FieldNameLength = FCString::Strlen(FieldName);
+
+	for (Json::FValue::ConstMemberIterator MemberIt = Object.MemberBegin(); MemberIt != Object.MemberEnd(); ++MemberIt)
+	{
+		if (FieldNameLength != MemberIt->name.GetStringLength())
+		{
+			continue;
+		}
+
+		if (FCString::Strnicmp(FieldName, MemberIt->name.GetString(), FieldNameLength) == 0)
+		{
+			return MemberIt;
+		}
+	}
+
+	return Object.MemberEnd();
+}
+
 bool TryGetBoolField(Json::FConstObject Object, const TCHAR* FieldName, bool& Out)
 {
-	Json::FValue::ConstMemberIterator It = Object.FindMember(FieldName);
+	Json::FValue::ConstMemberIterator It = FindMemberCaseInsensitive(Object, FieldName);
 	if (It == Object.MemberEnd())
 	{
 		return false;
@@ -98,7 +119,7 @@ static inline bool TryGetIntHelper(Json::FConstObject Object, const TCHAR* Field
 {
 	using namespace Json;
 
-	FValue::ConstMemberIterator It = Object.FindMember(FieldName);
+	FValue::ConstMemberIterator It = FindMemberCaseInsensitive(Object, FieldName);
 	if (It == Object.MemberEnd())
 	{
 		return false;
@@ -201,7 +222,7 @@ bool TryGetStringField(Json::FConstObject Object, const TCHAR* FieldName, FStrin
 {
 	using namespace Json;
 
-	FValue::ConstMemberIterator It = Object.FindMember(FieldName);
+	FValue::ConstMemberIterator It = FindMemberCaseInsensitive(Object, FieldName);
 	if (It == Object.MemberEnd())
 	{
 		return false;
@@ -221,7 +242,7 @@ bool TryGetStringArrayField(Json::FConstObject Object, const TCHAR* FieldName, T
 
 	Out.Reset();
 
-	FValue::ConstMemberIterator It = Object.FindMember(FieldName);
+	FValue::ConstMemberIterator It = FindMemberCaseInsensitive(Object, FieldName);
 	if (It == Object.MemberEnd() || !It->value.IsArray())
 	{
 		return false;
@@ -247,7 +268,7 @@ bool TryGetStringArrayField(Json::FConstObject Object, const TCHAR* FieldName, T
 
 	Out.Reset();
 
-	FValue::ConstMemberIterator It = Object.FindMember(FieldName);
+	FValue::ConstMemberIterator It = FindMemberCaseInsensitive(Object, FieldName);
 	if (It == Object.MemberEnd() || !It->value.IsArray())
 	{
 		return false;
diff --git a/Engine/Source/Runtime/Projects/Private/Tests/RapidJsonPluginLoadingTests.cpp b/Engine/Source/Runtime/Projects/Private/Tests/RapidJsonPluginLoadingTests.cpp
--- a/Engine/Source/Runtime/Projects/Private/Tests/RapidJsonPluginLoadingTests.cpp
+++ b/Engine/Source/Runtime/Projects/Private/Tests/RapidJsonPluginLoadingTests.cpp
@@ -446,4 +446,39 @@ TEST_CASE_NAMED(FProjectsLoadingProjectDescriptorTests, "System::Engine::Project
 
 }
 
+TEST_CASE_NAMED(FProjectsLoadingRapidJsonPluginLoadingCaseInsensitiveKeys, "System::Engine::Projects::Loading::RapidJsonPluginLoading::CaseInsensitiveKeys", "[ApplicationContextMask][SmokeFilter]")
+{
+	const TCHAR JsonText[] = TEXT(R"json(
+{
+	"aaa": 1,
+	"bBb": 2,
+	"CCC": 3,
+})json");
+
+
+	FProjectDescriptor Descriptor;
+
+	using namespace UE::Json;
+	using namespace UE::Projects::Private;
+
+	TValueOrError<FDocument, FParseError> RapidJsonDocument = UE::Json::Parse(JsonText);
+	FString ParseErrorMessage = RapidJsonDocument.HasError() ? RapidJsonDocument.GetError().CreateMessage(JsonText) : TEXT("No Error");
+	REQUIRE_MESSAGE(ParseErrorMessage, RapidJsonDocument.HasValue());
+
+	TOptional<UE::Json::FConstObject> RapidRoot = UE::Json::GetRootObject(RapidJsonDocument.GetValue());
+	REQUIRE(RapidRoot.IsSet());
+
+	int32 AAA = -1;
+	int32 BBB = -1;
+	int32 CCC = -1;
+
+	REQUIRE(TryGetNumberField(*RapidRoot, TEXT("Aaa"), AAA));
+	REQUIRE(TryGetNumberField(*RapidRoot, TEXT("Bbb"), BBB));
+	REQUIRE(TryGetNumberField(*RapidRoot, TEXT("Ccc"), CCC));
+
+	REQUIRE(AAA == 1);
+	REQUIRE(BBB == 2);
+	REQUIRE(CCC == 3);
+}
+
 #endif // WITH_TESTS
