[Plugin] Remove plugin from deferred GC if plugin got reloaded before batch finishes
#rnx

#Synced-CL 50080770

diff --git a/Engine/Source/Runtime/CoreUObject/Private/Misc/CoreUObjectPluginManager.cpp b/Engine/Source/Runtime/CoreUObject/Private/Misc/CoreUObjectPluginManager.cpp
--- a/Engine/Source/Runtime/CoreUObject/Private/Misc/CoreUObjectPluginManager.cpp
+++ b/Engine/Source/Runtime/CoreUObject/Private/Misc/CoreUObjectPluginManager.cpp
@@ -290,6 +290,16 @@ namespace UE::CoreUObject::Private
 		}
 	}
 
+	void PluginHandler::OnPluginLoad(IPlugin& Plugin)
+	{
+		check(IsInGameThread());
+
+		if (SuppressGCRefCount > 0)
+		{
+			DeferredPluginsToGC.RemoveSwap(Plugin.GetName());
+		}
+	}
+
 	void PluginHandler::OnPluginUnload(IPlugin& Plugin)
 	{
 		check(IsInGameThread());
diff --git a/Engine/Source/Runtime/CoreUObject/Private/Misc/CoreUObjectPluginManager.h b/Engine/Source/Runtime/CoreUObject/Private/Misc/CoreUObjectPluginManager.h
--- a/Engine/Source/Runtime/CoreUObject/Private/Misc/CoreUObjectPluginManager.h
+++ b/Engine/Source/Runtime/CoreUObject/Private/Misc/CoreUObjectPluginManager.h
@@ -12,6 +12,7 @@ namespace UE::CoreUObject::Private
 	public:
 		static void Install();
 
+		virtual void OnPluginLoad(IPlugin& Plugin) override;
 		virtual void OnPluginUnload(IPlugin& Plugin) override;
 
 		virtual void SuppressPluginUnloadGC() override;
diff --git a/Engine/Source/Runtime/Projects/Internal/ICoreUObjectPluginManager.h b/Engine/Source/Runtime/Projects/Internal/ICoreUObjectPluginManager.h
--- a/Engine/Source/Runtime/Projects/Internal/ICoreUObjectPluginManager.h
+++ b/Engine/Source/Runtime/Projects/Internal/ICoreUObjectPluginManager.h
@@ -12,6 +12,7 @@ namespace UE::PluginManager::Private
 	class ICoreUObjectPluginManager
 	{
 	public:
+		virtual void OnPluginLoad(IPlugin& Plugin) = 0;
 		virtual void OnPluginUnload(IPlugin& Plugin) = 0;
 
 		/**
diff --git a/Engine/Source/Runtime/Projects/Private/PluginManager.cpp b/Engine/Source/Runtime/Projects/Private/PluginManager.cpp
--- a/Engine/Source/Runtime/Projects/Private/PluginManager.cpp
+++ b/Engine/Source/Runtime/Projects/Private/PluginManager.cpp
@@ -3530,6 +3530,11 @@ void FPluginManager::MountPluginFromExternalSource(const TSharedRef<FPlugin>& Pl
 	const bool bHasContentOrVerse = (Plugin->CanContainContent() || Plugin->CanContainVerse()) && ensure(RegisterMountPointDelegate.IsBound());
 	if (bHasContentOrVerse)
 	{
+		if (UE::PluginManager::Private::CoreUObjectPluginHandler)
+		{
+			UE::PluginManager::Private::CoreUObjectPluginHandler->OnPluginLoad(*Plugin);
+		}
+
 		if (NewPluginMountedEvent.IsBound())
 		{
 			NewPluginMountedEvent.Broadcast(*Plugin);
