More MSVC PGO uninitialized variable workarounds

#Synced-CL 50878324

diff --git a/Engine/Source/Runtime/Projects/Private/PluginDescriptor.cpp b/Engine/Source/Runtime/Projects/Private/PluginDescriptor.cpp
--- a/Engine/Source/Runtime/Projects/Private/PluginDescriptor.cpp
+++ b/Engine/Source/Runtime/Projects/Private/PluginDescriptor.cpp
@@ -355,7 +355,7 @@ TOptional<FText> UE::Projects::Private::Read(UE::Json::FConstObject Object, FPlu
 
 	TryGetBoolField(Object, TEXT("EnableIAD"), Result.bEnableIAD);
 
-	bool bEnabledByDefault;
+	bool bEnabledByDefault{};
 	if(TryGetBoolField(Object, TEXT("EnabledByDefault"), bEnabledByDefault))
 	{
 		Result.EnabledByDefault = bEnabledByDefault ? EPluginEnabledByDefault::Enabled : EPluginEnabledByDefault::Disabled;
@@ -373,7 +373,7 @@ TOptional<FText> UE::Projects::Private::Read(UE::Json::FConstObject Object, FPlu
 	TryGetBoolField(Object, TEXT("ExplicitlyLoaded"), Result.bExplicitlyLoaded);
 	TryGetBoolField(Object, TEXT("HasExplicitPlatforms"), Result.bHasExplicitPlatforms);
 
-	bool bCanBeUsedWithUnrealHeaderTool;
+	bool bCanBeUsedWithUnrealHeaderTool{};
 	if(TryGetBoolField(Object, TEXT("CanBeUsedWithUnrealHeaderTool"), bCanBeUsedWithUnrealHeaderTool) && bCanBeUsedWithUnrealHeaderTool)
 	{
 		Result.SupportedPrograms.Add(TEXT("UnrealHeaderTool"));
