Profiling scopes

#Synced-CL 50998084

#Synced-CL 50998084

diff --git a/Engine/Plugins/Runtime/RigVM/Source/RigVM/Private/RigVMCore/RigVMRegistry.cpp b/Engine/Plugins/Runtime/RigVM/Source/RigVM/Private/RigVMCore/RigVMRegistry.cpp
--- a/Engine/Plugins/Runtime/RigVM/Source/RigVM/Private/RigVMCore/RigVMRegistry.cpp
+++ b/Engine/Plugins/Runtime/RigVM/Source/RigVM/Private/RigVMCore/RigVMRegistry.cpp
@@ -1726,6 +1726,7 @@ void FRigVMRegistry_RWLock::OnAssetRemoved(const FAssetData& InAssetData)
 
 void FRigVMRegistry_RWLock::OnPluginLoaded(IPlugin& InPlugin)
 {
+	TRACE_CPUPROFILER_EVENT_SCOPE_STR("FRigVMRegistry_RWLock::OnPluginLoaded");
 	bool bRegistryChanged;
 	{
 		FConditionalWriteScopeLock _(*this);
diff --git a/Engine/Plugins/Runtime/RigVM/Source/RigVM/Private/RigVMTypeUtils.cpp b/Engine/Plugins/Runtime/RigVM/Source/RigVM/Private/RigVMTypeUtils.cpp
--- a/Engine/Plugins/Runtime/RigVM/Source/RigVM/Private/RigVMTypeUtils.cpp
+++ b/Engine/Plugins/Runtime/RigVM/Source/RigVM/Private/RigVMTypeUtils.cpp
@@ -824,6 +824,7 @@ FRigVMPluginsPrefixCache::~FRigVMPluginsPrefixCache()
 
 void FRigVMPluginsPrefixCache::OnPluginMounted(IPlugin& InPlugin)
 {
+	TRACE_CPUPROFILER_EVENT_SCOPE_STR("FRigVMPluginsPrefixCache::OnPluginLoaded");
 	const TArray<FString> Prefixes = RigVMTypeUtils::GetPluginModulePrefixes(InPlugin);
 	for (const FString& Prefix : Prefixes)
 	{
diff --git a/Engine/Source/Runtime/CoreUObject/Private/Misc/CoreUObjectPluginManager.cpp b/Engine/Source/Runtime/CoreUObject/Private/Misc/CoreUObjectPluginManager.cpp
--- a/Engine/Source/Runtime/CoreUObject/Private/Misc/CoreUObjectPluginManager.cpp
+++ b/Engine/Source/Runtime/CoreUObject/Private/Misc/CoreUObjectPluginManager.cpp
@@ -292,6 +292,7 @@ namespace UE::CoreUObject::Private
 
 	void PluginHandler::OnPluginLoad(IPlugin& Plugin)
 	{
+		TRACE_CPUPROFILER_EVENT_SCOPE_STR("CoreUObjectPluginManager OnPluginLoad");
 		check(IsInGameThread());
 
 		if (SuppressGCRefCount > 0)
@@ -302,6 +303,7 @@ namespace UE::CoreUObject::Private
 
 	void PluginHandler::OnPluginUnload(IPlugin& Plugin)
 	{
+		TRACE_CPUPROFILER_EVENT_SCOPE_STR("CoreUObjectPluginManager OnPluginUnload");
 		check(IsInGameThread());
 
 		if (SuppressGCRefCount > 0)
@@ -324,6 +326,7 @@ namespace UE::CoreUObject::Private
 
 	void PluginHandler::ResumePluginUnloadGC()
 	{
+		TRACE_CPUPROFILER_EVENT_SCOPE_STR("CoreUObjectPluginManager ResumePluginUnloadGC");
 		check(IsInGameThread());
 
 		--SuppressGCRefCount;
diff --git a/Engine/Source/Runtime/Projects/Private/PluginManager.cpp b/Engine/Source/Runtime/Projects/Private/PluginManager.cpp
--- a/Engine/Source/Runtime/Projects/Private/PluginManager.cpp
+++ b/Engine/Source/Runtime/Projects/Private/PluginManager.cpp
@@ -3543,11 +3543,15 @@ void FPluginManager::MountPluginFromExternalSource(const TSharedRef<FPlugin>& Pl
 
 		if (NewPluginMountedEvent.IsBound())
 		{
+			TRACE_CPUPROFILER_EVENT_SCOPE_STR("OnNewPluginMounted");
 			NewPluginMountedEvent.Broadcast(*Plugin);
 		}
 
 		FString ContentDir = Plugin->GetContentDir();
-		RegisterMountPointDelegate.Execute(Plugin->GetMountedAssetPath(), ContentDir);
+		{ 
+			TRACE_CPUPROFILER_EVENT_SCOPE_STR("RegisterMountPoint");
+			RegisterMountPointDelegate.Execute(Plugin->GetMountedAssetPath(), ContentDir);
+		}
 
 		// Register this plugin's path with the list of content directories that the editor will search
 		if (Plugin->CanContainContent())
@@ -3562,6 +3566,7 @@ void FPluginManager::MountPluginFromExternalSource(const TSharedRef<FPlugin>& Pl
 	// Notify that additional localization data should be loaded
 	if (!Plugin->Descriptor.bExplicitlyLoaded && Plugin->Descriptor.LocalizationTargets.Num() > 0)
 	{
+		TRACE_CPUPROFILER_EVENT_SCOPE_STR("Mount Localization Paths");
 		TArray<FString> AdditionalLocResPaths;
 		PluginLocalizationUtils::GetLocalizationPathsForPlugin(*Plugin, AdditionalLocResPaths);
 		FTextLocalizationManager::Get().HandleLocalizationTargetsMounted(AdditionalLocResPaths);
@@ -3570,6 +3575,7 @@ void FPluginManager::MountPluginFromExternalSource(const TSharedRef<FPlugin>& Pl
 	// If it's a code module, also load the modules for it
 	if (Plugin->Descriptor.Modules.Num() > 0)
 	{
+		TRACE_CPUPROFILER_EVENT_SCOPE_STR("Load Plugin Modules");
 		// Add the plugin binaries directory
 		const FString PluginBinariesPath = FPaths::Combine(*FPaths::GetPath(Plugin->FileName), TEXT("Binaries"), FPlatformProcess::GetBinariesSubdirectory());
 		FModuleManager::Get().AddBinariesDirectory(*PluginBinariesPath, Plugin->GetLoadedFrom() == EPluginLoadedFrom::Project);
@@ -3589,6 +3595,7 @@ void FPluginManager::MountPluginFromExternalSource(const TSharedRef<FPlugin>& Pl
 	// Notify listeners that the plugin is completely mounted now
 	if (bHasContentOrVerse && NewPluginContentMountedEvent.IsBound())
 	{
+		TRACE_CPUPROFILER_EVENT_SCOPE_STR("OnNewPluginContentMounted");
 		NewPluginContentMountedEvent.Broadcast(*Plugin);
 	}
 
diff --git a/Engine/Source/Runtime/RenderCore/Private/ShaderLibrary/ShaderCodeLibrary.cpp b/Engine/Source/Runtime/RenderCore/Private/ShaderLibrary/ShaderCodeLibrary.cpp
--- a/Engine/Source/Runtime/RenderCore/Private/ShaderLibrary/ShaderCodeLibrary.cpp
+++ b/Engine/Source/Runtime/RenderCore/Private/ShaderLibrary/ShaderCodeLibrary.cpp
@@ -2069,6 +2069,7 @@ FShaderLibrariesCollection* FShaderLibrariesCollection::Impl = nullptr;
 
 static void FShaderCodeLibraryPluginMountedCallback(IPlugin& Plugin)
 {
+	TRACE_CPUPROFILER_EVENT_SCOPE_STR("FShaderCodeLibraryPluginMountedCallback");
 	if (FApp::CanEverRender() && UE::ShaderLibrary::Private::PluginsToIgnoreOnMount.Remove(Plugin.GetName()) == 0)
 	{
 		FShaderCodeLibrary::OpenPluginShaderLibrary(Plugin);
