Remove a very spammy boot timing scope. ConfigureEnabledPluginForCurrentTarget was hit ~7000 times during boot but its time is insignificant.

#Synced-CL 51003583

diff --git a/Engine/Source/Runtime/Projects/Private/PluginManager.cpp b/Engine/Source/Runtime/Projects/Private/PluginManager.cpp
--- a/Engine/Source/Runtime/Projects/Private/PluginManager.cpp
+++ b/Engine/Source/Runtime/Projects/Private/PluginManager.cpp
@@ -2189,8 +2189,6 @@ bool FPluginManager::GetCodePluginsForProject(const FProjectDescriptor* ProjectD
 bool FPluginManager::ConfigureEnabledPluginForCurrentTarget(const FPluginReferenceDescriptor& FirstReference,
 	TMap<FString, FPlugin*>& EnabledPlugins, FStringView SourceOfPluginRequest, const TSet<FString>& AllowedOptionalDependencies)
 {
-	SCOPED_BOOT_TIMING("ConfigureEnabledPluginForCurrentTarget");
-
 	FConfigurePluginResultInfo ResultInfo;
 
 	if (ConfigureEnabledPluginForTarget(FirstReference, IProjectManager::Get().GetCurrentProject(), UE_APP_NAME,
