Add a way for plugins to opt into or out of IAS, defaulting to enabled.

#rnx

- Adding '"EnableIAS" : false' to the uplugin file will opt the plugin our of IAS.
#Synced-CL 51774074

#Synced-CL 51774074

diff --git a/Engine/Source/Runtime/Projects/Private/PluginDescriptor.cpp b/Engine/Source/Runtime/Projects/Private/PluginDescriptor.cpp
--- a/Engine/Source/Runtime/Projects/Private/PluginDescriptor.cpp
+++ b/Engine/Source/Runtime/Projects/Private/PluginDescriptor.cpp
@@ -354,6 +354,7 @@ TOptional<FText> UE::Projects::Private::Read(UE::Json::FConstObject Object, FPlu
 	TryGetBoolField(Object, TEXT("EnableVerseAssetReflection"), Result.bEnableVerseAssetReflection);
 
 	TryGetBoolField(Object, TEXT("EnableIAD"), Result.bEnableIAD);
+	TryGetBoolField(Object, TEXT("EnableIAS"), Result.bEnableIAS);
 
 	bool bEnabledByDefault{};
 	if(TryGetBoolField(Object, TEXT("EnabledByDefault"), bEnabledByDefault))
@@ -549,6 +550,15 @@ void FPluginDescriptor::UpdateJson(FJsonObject& JsonObject) const
 		JsonObject.RemoveField(TEXT("EnableIAD"));
 	}
 
+	if (!bEnableIAS)
+	{
+		JsonObject.SetBoolField(TEXT("EnableIAS"), bEnableIAS);
+	}
+	else
+	{
+		JsonObject.RemoveField(TEXT("EnableIAS"));
+	}
+
 	if (EnabledByDefault != EPluginEnabledByDefault::Unspecified)
 	{
 		JsonObject.SetBoolField(TEXT("EnabledByDefault"), (EnabledByDefault == EPluginEnabledByDefault::Enabled));
diff --git a/Engine/Source/Runtime/Projects/Public/PluginDescriptor.h b/Engine/Source/Runtime/Projects/Public/PluginDescriptor.h
--- a/Engine/Source/Runtime/Projects/Public/PluginDescriptor.h
+++ b/Engine/Source/Runtime/Projects/Public/PluginDescriptor.h
@@ -114,6 +114,9 @@ struct FPluginDescriptor
 	/** Use IAD for this plugin */
 	bool bEnableIAD = false;
 
+	/** Use IAS for this plugin. This is likely to be removed in the near future. */
+	bool bEnableIAS = true;
+
 	/** Whether this plugin should be enabled by default for all projects */
 	EPluginEnabledByDefault EnabledByDefault;
 
