Add new IsRootPlugin function.
A root plugin is one that is meant to be used as a way to define a collection for other plugins.
Root plugins should have no code or content only a list of dependencies that define the plugins for this root.

diff --git a/Engine/Source/Editor/UnrealEd/Private/CookOnTheFlyServer.cpp b/Engine/Source/Editor/UnrealEd/Private/CookOnTheFlyServer.cpp
--- a/Engine/Source/Editor/UnrealEd/Private/CookOnTheFlyServer.cpp
+++ b/Engine/Source/Editor/UnrealEd/Private/CookOnTheFlyServer.cpp
@@ -9081,15 +9081,13 @@ void UCookOnTheFlyServer::WriteCookMetadata(const ITargetPlatform* InTargetPlatf
 			// bogus data, it won't get saved.
 			Entry.DependencyIndexStart = (uint32)PluginChildArray.Num();
 
-			const FPluginDescriptor& Descriptor = EnabledPlugin->GetDescriptor();
-
-			// Root plugins are sealed && no code
-			if (Descriptor.bIsSealed && Descriptor.bNoCode)
+			if (EnabledPlugin->IsRootPlugin())
 			{
 				Entry.Type = UE::Cook::ECookMetadataPluginType::Root;
 				RootPlugins.Add(SelfIndex);
 			}
 
+			const FPluginDescriptor& Descriptor = EnabledPlugin->GetDescriptor();
 			// Pull in any custom fields the project wants to pass on.
 			auto CheckForCustomFields = [&Descriptor, &PlatformNameString, &EnabledPlugin, &Entry, &CustomFieldNameIndex](bool bIsBool, bool bIsPerPlatform, const TArray<FString>& FieldNames)
 			{
diff --git a/Engine/Source/Runtime/Projects/Private/PluginDescriptor.cpp b/Engine/Source/Runtime/Projects/Private/PluginDescriptor.cpp
--- a/Engine/Source/Runtime/Projects/Private/PluginDescriptor.cpp
+++ b/Engine/Source/Runtime/Projects/Private/PluginDescriptor.cpp
@@ -367,6 +367,7 @@ TOptional<FText> UE::Projects::Private::Read(UE::Json::FConstObject Object, FPlu
 	TryGetBoolField(Object, TEXT("CanContainContent"), Result.bCanContainContent);
 	TryGetBoolField(Object, TEXT("CanContainVerse"), Result.bCanContainVerse);
 	TryGetBoolField(Object, TEXT("NoCode"), Result.bNoCode);
+	TryGetBoolField(Object, TEXT("NoGameFeatureContent"), Result.bNoGameFeatureContent);
 	TryGetBoolField(Object, TEXT("IsBetaVersion"), Result.bIsBetaVersion);
 	TryGetBoolField(Object, TEXT("IsExperimentalVersion"), Result.bIsExperimentalVersion);
 	TryGetBoolField(Object, TEXT("Installed"), Result.bInstalled);
@@ -593,6 +594,10 @@ void FPluginDescriptor::UpdateJson(FJsonObject& JsonObject) const
 	{
 		JsonObject.SetBoolField(TEXT("NoCode"), bNoCode);
 	}
+	if (bNoGameFeatureContent)
+	{
+		JsonObject.SetBoolField(TEXT("NoGameFeatureContent"), bNoGameFeatureContent);
+	}
 	JsonObject.SetBoolField(TEXT("IsBetaVersion"), bIsBetaVersion);
 	JsonObject.SetBoolField(TEXT("IsExperimentalVersion"), bIsExperimentalVersion);
 	JsonObject.SetBoolField(TEXT("Installed"), bInstalled);
diff --git a/Engine/Source/Runtime/Projects/Private/PluginManager.h b/Engine/Source/Runtime/Projects/Private/PluginManager.h
--- a/Engine/Source/Runtime/Projects/Private/PluginManager.h
+++ b/Engine/Source/Runtime/Projects/Private/PluginManager.h
@@ -131,6 +131,11 @@ public:
 		return Descriptor.bEnableSceneGraphAnimation;
 	}
 
+	virtual bool IsRootPlugin() const override
+	{
+		return Descriptor.bNoCode && Descriptor.bNoGameFeatureContent;
+	}
+
 	virtual EPluginType GetType() const override
 	{
 		return Type;
diff --git a/Engine/Source/Runtime/Projects/Public/Interfaces/IPluginManager.h b/Engine/Source/Runtime/Projects/Public/Interfaces/IPluginManager.h
--- a/Engine/Source/Runtime/Projects/Public/Interfaces/IPluginManager.h
+++ b/Engine/Source/Runtime/Projects/Public/Interfaces/IPluginManager.h
@@ -239,6 +239,12 @@ public:
 	 */
 	virtual bool IsSceneGraphAnimationEnabled() const = 0;
 
+	/**
+	 * Returns true if this is a root plugin. A root plugin is meant as a collection of other plugins.
+	 * Root plugins should only have plugin dependencies and be used to designate some group of plugins.
+	 */
+	virtual bool IsRootPlugin() const = 0;
+
 	/**
 	 * Returns the plugin's location
 	 *
diff --git a/Engine/Source/Runtime/Projects/Public/PluginDescriptor.h b/Engine/Source/Runtime/Projects/Public/PluginDescriptor.h
--- a/Engine/Source/Runtime/Projects/Public/PluginDescriptor.h
+++ b/Engine/Source/Runtime/Projects/Public/PluginDescriptor.h
@@ -159,6 +159,9 @@ struct FPluginDescriptor
 	/** If true, this plugin from a platform extension extending another plugin */
 	bool bIsPluginExtension;
 
+	/** If true, the GameFeaturePlugin this is on should have no content other than the GameFeatureData asset. */
+	bool bNoGameFeatureContent = false;
+
 	/** Pre-build steps for each host platform */
 	FCustomBuildSteps PreBuildSteps;
 
